#ttp 330490 - EDITOR: LAYERS: Giving a layer a duplicate name should not commit the text and should display a warning
#proj UE4
#branch UE4
#summary Improved the error reporting when renaming a layer

diff --git a/Engine/Source/Editor/Layers/Private/LayerViewModel.cpp b/Engine/Source/Editor/Layers/Private/LayerViewModel.cpp
--- a/Engine/Source/Editor/Layers/Private/LayerViewModel.cpp
+++ b/Engine/Source/Editor/Layers/Private/LayerViewModel.cpp
@@ -92,26 +92,43 @@ void FLayerViewModel::ToggleVisibility()
 }
 
 
-void FLayerViewModel::RenameTo( const FString& NewLayerName )
+bool FLayerViewModel::CanRenameTo( const FName& NewLayerName, FString& OutMessage ) const
+{
+	if (NewLayerName.IsNone())
+	{
+		OutMessage = LOCTEXT("EmptyLayerName", "Layer must be given a name").ToString();
+		return false;
+	}
+
+	TWeakObjectPtr< ULayer > FoundLayer;
+	if ( WorldLayers->TryGetLayer( NewLayerName, FoundLayer ) && FoundLayer != Layer )
+	{
+		OutMessage = LOCTEXT("RenameFailed_AlreadyExists", "This layer already exists").ToString();
+		return false;
+	}
+
+	return true;
+}
+
+
+void FLayerViewModel::RenameTo( const FName& NewLayerName )
 {
 	if( !Layer.IsValid() )
 	{
 		return;
 	}
 
-	const FName NewLayerFName = FName( *NewLayerName );
-
-	if( Layer->LayerName == NewLayerFName)
+	if( Layer->LayerName == NewLayerName)
 	{
 		return;
 	}
 
 	int32 LayerIndex = 0;
-	FName UniqueNewLayerName = NewLayerFName;
+	FName UniqueNewLayerName = NewLayerName;
 	TWeakObjectPtr< ULayer > FoundLayer;
 	while( WorldLayers->TryGetLayer( UniqueNewLayerName, FoundLayer ) )
 	{
-		UniqueNewLayerName = FName( *FString::Printf( TEXT( "%s_%d" ), *NewLayerFName.ToString(), ++LayerIndex ) );
+		UniqueNewLayerName = FName( *FString::Printf( TEXT( "%s_%d" ), *NewLayerName.ToString(), ++LayerIndex ) );
 	}
 
 	const FScopedTransaction Transaction( LOCTEXT("RenameTo", "Rename Layer") );
diff --git a/Engine/Source/Editor/Layers/Private/LayerViewModel.h b/Engine/Source/Editor/Layers/Private/LayerViewModel.h
--- a/Engine/Source/Editor/Layers/Private/LayerViewModel.h
+++ b/Engine/Source/Editor/Layers/Private/LayerViewModel.h
@@ -48,12 +48,21 @@ public:
 	 */
 	void ToggleVisibility();
 
+	/**
+	 *	Returns whether the layer can be assigned the specified name
+	 *
+	 *	@param	NewLayerName		The Layers new name
+	 *	@param	OutMessage [OUT]	A returned description explaining the boolean result
+	 *	@return						true then the name can be assigned
+	 */
+	bool CanRenameTo( const FName& NewLayerName, FString& OutMessage ) const;
+
 	/**
 	 *	Renames the Layer to the specified name
 	 *
 	 *	@param	NewLayerName	The Layers new name
 	 */
-	void RenameTo( const FString& NewLayerName );
+	void RenameTo( const FName& NewLayerName );
 
 	/**
 	 *	Returns whether the specified actors can be assigned to the Layer
diff --git a/Engine/Source/Editor/Layers/Private/SLayersViewRow.h b/Engine/Source/Editor/Layers/Private/SLayersViewRow.h
--- a/Engine/Source/Editor/Layers/Private/SLayersViewRow.h
+++ b/Engine/Source/Editor/Layers/Private/SLayersViewRow.h
@@ -121,16 +121,17 @@ protected:
 	{
 		if (!InText.IsEmpty())
 		{
-			ViewModel->RenameTo(InText.ToString());
+			ViewModel->RenameTo(*InText.ToString());
 		}
 	}
 
 	/** Callback when the SInlineEditableTextBlock is changed, to check for error conditions. */
 	bool OnRenameLayerTextChanged(const FText& NewText, FText& OutErrorMessage)
 	{
-		if (NewText.IsEmpty())
+		FString OutMessage;
+		if ( !ViewModel->CanRenameTo( *NewText.ToString(), OutMessage ) )
 		{
-			OutErrorMessage = LOCTEXT("EmptyLayerName", "Layer must be given a name");
+			OutErrorMessage = FText::FromString( OutMessage );
 			return false;
 		}
 
