Changed layer and TileLayer visibility icons to be the same as the scene outliner visibility icons, for consistency.

diff --git a/Engine/Source/Editor/Layers/Private/SLayersView.h b/Engine/Source/Editor/Layers/Private/SLayersView.h
--- a/Engine/Source/Editor/Layers/Private/SLayersView.h
+++ b/Engine/Source/Editor/Layers/Private/SLayersView.h
@@ -59,14 +59,14 @@ public:
 			/** We don't want the normal header to be visible */
 			.Visibility( EVisibility::Collapsed )
 
-			/** LayerName label column */
-			+ SHeaderRow::Column( LayersView::ColumnID_LayerLabel )
-			.DefaultLabel( LOCTEXT("Column_LayerNameLabel", "Layer") )
-
 			/** Layer visibility column */
-			+ SHeaderRow::Column( LayersView::ColumnID_Visibility )
-			.DefaultLabel( NSLOCTEXT("LayersView", "Visibility", "Visibility") )
-			.FixedWidth( 48.0f );
+			+SHeaderRow::Column(LayersView::ColumnID_Visibility)
+			.DefaultLabel(NSLOCTEXT("LayersView", "Visibility", "Visibility"))
+			.FixedWidth(40.0f)
+
+			/** LayerName label column */
+			+SHeaderRow::Column(LayersView::ColumnID_LayerLabel)
+			.DefaultLabel(LOCTEXT("Column_LayerNameLabel", "Layer"));
 
 		ChildSlot
 			[
diff --git a/Engine/Source/Editor/Layers/Private/SLayersViewRow.cpp b/Engine/Source/Editor/Layers/Private/SLayersViewRow.cpp
--- a/Engine/Source/Editor/Layers/Private/SLayersViewRow.cpp
+++ b/Engine/Source/Editor/Layers/Private/SLayersViewRow.cpp
@@ -59,17 +59,15 @@ TSharedRef< SWidget > SLayersViewRow::GenerateWidgetForColumn(const FName& Colum
 		TableRowContent =
 			SAssignNew(VisibilityButton, SButton)
 			.ContentPadding(0)
-			.ButtonStyle(FEditorStyle::Get(), "ToggleButton")
+			.ButtonStyle(FEditorStyle::Get(), "NoBorder")
 			.OnClicked(this, &SLayersViewRow::OnToggleVisibility)
 			.ToolTipText(LOCTEXT("VisibilityButtonToolTip", "Toggle Layer Visibility"))
-			.ForegroundColor(FSlateColor::UseForeground())
 			.HAlign(HAlign_Center)
 			.VAlign(VAlign_Center)
 			.Content()
 			[
 				SNew(SImage)
 				.Image(this, &SLayersViewRow::GetVisibilityBrushForLayer)
-				.ColorAndOpacity(this, &SLayersViewRow::GetForegroundColorForButton)
 			]
 		;
 	}
@@ -154,12 +152,6 @@ FReply SLayersViewRow::OnDrop(const FGeometry& MyGeometry, const FDragDropEvent&
 	return FReply::Handled();
 }
 
-FSlateColor SLayersViewRow::GetForegroundColorForButton() const
-{
-	static const FName InvertedForegroundName("InvertedForeground");
-	return (VisibilityButton.IsValid() && (VisibilityButton->IsHovered() || VisibilityButton->IsPressed())) ? FEditorStyle::GetSlateColor(InvertedForegroundName) : FSlateColor::UseForeground();
-}
-
 FSlateColor SLayersViewRow::GetColorAndOpacity() const
 {
 	if (!FSlateApplication::Get().IsDragDropping())
@@ -183,7 +175,16 @@ FSlateColor SLayersViewRow::GetColorAndOpacity() const
 
 const FSlateBrush* SLayersViewRow::GetVisibilityBrushForLayer() const
 {
-	return (ViewModel->IsVisible()) ? FEditorStyle::GetBrush("Layer.VisibleIcon16x") : FEditorStyle::GetBrush("Layer.NotVisibleIcon16x");
+	if (ViewModel->IsVisible())
+	{
+		return IsHovered() ? FEditorStyle::GetBrush("Level.VisibleHighlightIcon16x") :
+			FEditorStyle::GetBrush("Level.VisibleIcon16x");
+	}
+	else
+	{
+		return IsHovered() ? FEditorStyle::GetBrush("Level.NotVisibleHighlightIcon16x") :
+			FEditorStyle::GetBrush("Level.NotVisibleIcon16x");
+	}
 }
 
 #undef LOCTEXT_NAMESPACE
diff --git a/Engine/Source/Editor/Layers/Private/SLayersViewRow.h b/Engine/Source/Editor/Layers/Private/SLayersViewRow.h
--- a/Engine/Source/Editor/Layers/Private/SLayersViewRow.h
+++ b/Engine/Source/Editor/Layers/Private/SLayersViewRow.h
@@ -81,14 +81,6 @@ protected:
 
 private:
 	
-	/**
-	 *	Gets the appropriate SlateColor for the specified button depending on the button's current state
-	 *
-	 *	@param	Button	The SButton Widget to get the Foreground Color for
-	 *	@return			The ForegroundColor
-	 */
-	FSlateColor GetForegroundColorForButton() const;
-
 	/**
 	 *	Returns the Color and Opacity for displaying the bound Layer's Name.
 	 *	The Color and Opacity changes depending on whether a drag/drop operation is occurring
