Fix Non-Unity Missing Include

diff --git a/Engine/Source/Editor/Layers/Private/ActorsAssignedToSpecificLayersFilter.h b/Engine/Source/Editor/Layers/Private/ActorsAssignedToSpecificLayersFilter.h
--- a/Engine/Source/Editor/Layers/Private/ActorsAssignedToSpecificLayersFilter.h
+++ b/Engine/Source/Editor/Layers/Private/ActorsAssignedToSpecificLayersFilter.h
@@ -4,6 +4,7 @@
 #include "CoreMinimal.h"
 #include "SceneOutlinerFilters.h"
 #include "ActorTreeItem.h"
+#include "GameFramework/Actor.h"
 
 class FActorsAssignedToSpecificLayersFilter : public SceneOutliner::TOutlinerFilter<SceneOutliner::FActorTreeItem>, public TSharedFromThis< FActorsAssignedToSpecificLayersFilter >
 {
