Fix Actors cannot be assigned to a Layer when dragged from World Outliner.

diff --git a/Engine/Source/Editor/Layers/Private/SLayerBrowser.h b/Engine/Source/Editor/Layers/Private/SLayerBrowser.h
--- a/Engine/Source/Editor/Layers/Private/SLayerBrowser.h
+++ b/Engine/Source/Editor/Layers/Private/SLayerBrowser.h
@@ -16,7 +16,8 @@
 #include "Editor/Layers/Private/ActorsAssignedToSpecificLayersFilter.h"
 #include "ISceneOutlinerColumn.h"
 #include "Editor/Layers/Private/SceneOutlinerLayerContentsColumn.h"
-#include "DragAndDrop/ActorDragDropGraphEdOp.h"
+#include "SceneOutlinerDragDrop.h"
+#include "DragAndDrop/ActorDragDropOp.h"
 #include "Editor/Layers/Private/SLayersView.h"
 #include "Editor/Layers/Private/SLayersCommandsMenu.h"
 
@@ -84,10 +85,14 @@ protected:
 	 */
 	virtual void OnDragLeave( const FDragDropEvent& DragDropEvent ) override
 	{
-		TSharedPtr< FActorDragDropGraphEdOp > DragActorOp = DragDropEvent.GetOperationAs< FActorDragDropGraphEdOp >();
-		if(DragActorOp.IsValid())
+		TSharedPtr< FSceneOutlinerDragDropOp > DragOp = DragDropEvent.GetOperationAs< FSceneOutlinerDragDropOp >();
+		if (DragOp.IsValid())
 		{
-			DragActorOp->ResetToDefaultToolTip();
+			TSharedPtr< FActorDragDropOp > ActorDragOp = DragOp->GetSubOp<FActorDragDropOp>();
+			if (ActorDragOp.IsValid())
+			{
+				ActorDragOp->ResetToDefaultToolTip();
+			}
 		}
 	}
 
@@ -101,49 +106,56 @@ protected:
 	 */
 	virtual FReply OnDragOver( const FGeometry& MyGeometry, const FDragDropEvent& DragDropEvent ) override
 	{
-		TSharedPtr< FActorDragDropGraphEdOp > DragActorOp = DragDropEvent.GetOperationAs< FActorDragDropGraphEdOp >();
-		if (!SelectedLayerViewModel->GetDataSource().IsValid() || !DragActorOp.IsValid())
-		{
-			return FReply::Unhandled();
-		}
-
-		const FGeometry LayerContentsHeaderGeometry = SWidget::FindChildGeometry( MyGeometry, LayerContentsHeader.ToSharedRef() );
-		bool bValidDrop = LayerContentsHeaderGeometry.IsUnderLocation( DragDropEvent.GetScreenSpacePosition() );
-
-		if( !bValidDrop && Mode == ELayerBrowserMode::LayerContents )
-		{
-			const FGeometry LayerContentsSectionGeometry = SWidget::FindChildGeometry( MyGeometry, LayerContentsSection.ToSharedRef() );
-			bValidDrop = LayerContentsSectionGeometry.IsUnderLocation( DragDropEvent.GetScreenSpacePosition() );
-		}
-
-		if( !bValidDrop )
-		{
-			return FReply::Unhandled();
-		}
-
-		if ( !DragActorOp.IsValid() || DragActorOp->Actors.Num() == 0 )
+		if (!SelectedLayerViewModel->GetDataSource().IsValid())
 		{
 			return FReply::Unhandled();
 		}
 
-		bool bCanAssign = false;
-		FText Message;
-		if( DragActorOp->Actors.Num() > 1 )
+		TSharedPtr< FSceneOutlinerDragDropOp > DragOp = DragDropEvent.GetOperationAs< FSceneOutlinerDragDropOp >();
+		if (DragOp.IsValid())
 		{
-			bCanAssign = SelectedLayerViewModel->CanAssignActors( DragActorOp->Actors, OUT Message );
-		}
-		else
-		{
-			bCanAssign = SelectedLayerViewModel->CanAssignActor( DragActorOp->Actors[ 0 ], OUT Message );
-		}
-
-		if ( bCanAssign )
-		{
-			DragActorOp->SetToolTip( FActorDragDropGraphEdOp::ToolTip_CompatibleGeneric, Message );
-		}
-		else
-		{
-			DragActorOp->SetToolTip( FActorDragDropGraphEdOp::ToolTip_IncompatibleGeneric, Message );
+			TSharedPtr< FActorDragDropOp > ActorDragOp = DragOp->GetSubOp<FActorDragDropOp>();
+			if (ActorDragOp.IsValid())
+			{		
+				const FGeometry LayerContentsHeaderGeometry = SWidget::FindChildGeometry( MyGeometry, LayerContentsHeader.ToSharedRef() );
+				bool bValidDrop = LayerContentsHeaderGeometry.IsUnderLocation( DragDropEvent.GetScreenSpacePosition() );
+
+				if( !bValidDrop && Mode == ELayerBrowserMode::LayerContents )
+				{
+					const FGeometry LayerContentsSectionGeometry = SWidget::FindChildGeometry( MyGeometry, LayerContentsSection.ToSharedRef() );
+					bValidDrop = LayerContentsSectionGeometry.IsUnderLocation( DragDropEvent.GetScreenSpacePosition() );
+				}
+
+				if( !bValidDrop )
+				{
+					return FReply::Unhandled();
+				}
+
+				if ( !ActorDragOp.IsValid() || ActorDragOp->Actors.Num() == 0 )
+				{
+					return FReply::Unhandled();
+				}
+
+				bool bCanAssign = false;
+				FText Message;
+				if(ActorDragOp->Actors.Num() > 1 )
+				{
+					bCanAssign = SelectedLayerViewModel->CanAssignActors(ActorDragOp->Actors, OUT Message );
+				}
+				else
+				{
+					bCanAssign = SelectedLayerViewModel->CanAssignActor(ActorDragOp->Actors[ 0 ], OUT Message );
+				}
+
+				if ( bCanAssign )
+				{
+					ActorDragOp->SetToolTip(Message, FEditorStyle::GetBrush(TEXT("Graph.ConnectorFeedback.OK")));
+				}
+				else
+				{
+					ActorDragOp->SetToolTip(Message, FEditorStyle::GetBrush(TEXT("Graph.ConnectorFeedback.Error")));
+				}
+			}
 		}
 
 		// We leave the event unhandled so the children of the ListView get a chance to grab the drag/drop
@@ -160,30 +172,37 @@ protected:
 	 */
 	virtual FReply OnDrop( const FGeometry& MyGeometry, const FDragDropEvent& DragDropEvent ) override
 	{
-		TSharedPtr< FActorDragDropGraphEdOp > DragActorOp = DragDropEvent.GetOperationAs< FActorDragDropGraphEdOp >();
-
-		if (!SelectedLayerViewModel->GetDataSource().IsValid() || !DragActorOp.IsValid())
+		if (!SelectedLayerViewModel->GetDataSource().IsValid())
 		{
 			return FReply::Unhandled();
 		}
 
-		const FGeometry LayerContentsHeaderGeometry = SWidget::FindChildGeometry( MyGeometry, LayerContentsHeader.ToSharedRef() );
-		bool bValidDrop = LayerContentsHeaderGeometry.IsUnderLocation( DragDropEvent.GetScreenSpacePosition() );
-
-		if( !bValidDrop && Mode == ELayerBrowserMode::LayerContents )
+		TSharedPtr< FSceneOutlinerDragDropOp > DragOp = DragDropEvent.GetOperationAs< FSceneOutlinerDragDropOp >();
+		if (DragOp.IsValid())
 		{
-			const FGeometry LayerContentsSectionGeometry = SWidget::FindChildGeometry( MyGeometry, LayerContentsSection.ToSharedRef() );
-			bValidDrop = LayerContentsSectionGeometry.IsUnderLocation( DragDropEvent.GetScreenSpacePosition() );
-		}
-
-		if( !bValidDrop )
-		{
-			return FReply::Unhandled();
+			TSharedPtr< FActorDragDropOp > ActorDragOp = DragOp->GetSubOp<FActorDragDropOp>();
+			if (ActorDragOp.IsValid())
+			{
+				const FGeometry LayerContentsHeaderGeometry = SWidget::FindChildGeometry(MyGeometry, LayerContentsHeader.ToSharedRef());
+				bool bValidDrop = LayerContentsHeaderGeometry.IsUnderLocation(DragDropEvent.GetScreenSpacePosition());
+
+				if (!bValidDrop && Mode == ELayerBrowserMode::LayerContents)
+				{
+					const FGeometry LayerContentsSectionGeometry = SWidget::FindChildGeometry(MyGeometry, LayerContentsSection.ToSharedRef());
+					bValidDrop = LayerContentsSectionGeometry.IsUnderLocation(DragDropEvent.GetScreenSpacePosition());
+				}
+
+				if (!bValidDrop)
+				{
+					return FReply::Unhandled();
+				}
+
+				SelectedLayerViewModel->AddActors(ActorDragOp->Actors);
+				return FReply::Handled();
+			}
 		}
 
-		SelectedLayerViewModel->AddActors( DragActorOp->Actors );
-
-		return FReply::Handled();
+		return FReply::Unhandled();
 	}
 
 
diff --git a/Engine/Source/Editor/Layers/Private/SLayersView.h b/Engine/Source/Editor/Layers/Private/SLayersView.h
--- a/Engine/Source/Editor/Layers/Private/SLayersView.h
+++ b/Engine/Source/Editor/Layers/Private/SLayersView.h
@@ -10,7 +10,8 @@
 #include "Widgets/SBoxPanel.h"
 #include "LayerCollectionViewCommands.h"
 #include "Widgets/Views/SHeaderRow.h"
-#include "DragAndDrop/ActorDragDropGraphEdOp.h"
+#include "SceneOutlinerDragDrop.h"
+#include "DragAndDrop/ActorDragDropOp.h"
 #include "Widgets/Views/STableViewBase.h"
 #include "Widgets/Views/STableRow.h"
 #include "Widgets/Views/SListView.h"
@@ -173,10 +174,14 @@ protected:
 	 */
 	virtual void OnDragLeave( const FDragDropEvent& DragDropEvent ) override
 	{
-		TSharedPtr< FActorDragDropGraphEdOp > DragActorOp = DragDropEvent.GetOperationAs< FActorDragDropGraphEdOp >();
-		if (DragActorOp.IsValid())
+		TSharedPtr< FSceneOutlinerDragDropOp > DragOp = DragDropEvent.GetOperationAs< FSceneOutlinerDragDropOp >();
+		if (DragOp.IsValid())
 		{
-			DragActorOp->ResetToDefaultToolTip();
+			TSharedPtr< FActorDragDropOp > ActorDragOp = DragOp->GetSubOp<FActorDragDropOp>();
+			if (ActorDragOp.IsValid())
+			{
+				ActorDragOp->ResetToDefaultToolTip();
+			}
 		}
 	}
 
@@ -190,10 +195,14 @@ protected:
 	 */
 	virtual FReply OnDragOver( const FGeometry& MyGeometry, const FDragDropEvent& DragDropEvent ) override
 	{
-		TSharedPtr< FActorDragDropGraphEdOp > DragActorOp = DragDropEvent.GetOperationAs< FActorDragDropGraphEdOp >();
-		if (DragActorOp.IsValid())
+		TSharedPtr< FSceneOutlinerDragDropOp > DragOp = DragDropEvent.GetOperationAs< FSceneOutlinerDragDropOp >();
+		if (DragOp.IsValid())
 		{
-			DragActorOp->SetToolTip(FActorDragDropGraphEdOp::ToolTip_CompatibleGeneric, LOCTEXT("OnDragOver", "Add Actors to New Layer"));
+			TSharedPtr< FActorDragDropOp > ActorDragOp = DragOp->GetSubOp<FActorDragDropOp>();
+			if (ActorDragOp.IsValid())
+			{
+				ActorDragOp->SetToolTip(LOCTEXT("OnDragOver", "Add Actors to New Layer"), FEditorStyle::GetBrush(TEXT("Graph.ConnectorFeedback.OK")));
+			}
 		}
 
 		// We leave the event unhandled so the children of the ListView get a chance to grab the drag/drop
@@ -210,15 +219,18 @@ protected:
 	 */
 	virtual FReply OnDrop( const FGeometry& MyGeometry, const FDragDropEvent& DragDropEvent ) override
 	{
-		TSharedPtr< FActorDragDropGraphEdOp > DragActorOp = DragDropEvent.GetOperationAs< FActorDragDropGraphEdOp >();
-		if ( !DragActorOp.IsValid() )	
+		TSharedPtr< FSceneOutlinerDragDropOp > DragOp = DragDropEvent.GetOperationAs< FSceneOutlinerDragDropOp >();
+		if ( DragOp.IsValid() )
 		{
-			return FReply::Unhandled();
+			TSharedPtr< FActorDragDropOp > ActorDragOp = DragOp->GetSubOp<FActorDragDropOp>();
+			if ( ActorDragOp.IsValid() )
+			{
+				ViewModel->AddActorsToNewLayer(ActorDragOp->Actors);
+				return FReply::Handled();
+			}
 		}
 
-		ViewModel->AddActorsToNewLayer( DragActorOp->Actors );
-
-		return FReply::Handled();
+		return FReply::Unhandled();
 	}
 
 	FReply OnDragRow(const FGeometry& MyGeometry, const FPointerEvent& MouseEvent)
diff --git a/Engine/Source/Editor/Layers/Private/SLayersViewRow.cpp b/Engine/Source/Editor/Layers/Private/SLayersViewRow.cpp
--- a/Engine/Source/Editor/Layers/Private/SLayersViewRow.cpp
+++ b/Engine/Source/Editor/Layers/Private/SLayersViewRow.cpp
@@ -6,7 +6,8 @@
 #include "Widgets/Input/SButton.h"
 #include "Widgets/Views/SListView.h"
 #include "EditorStyleSet.h"
-#include "DragAndDrop/ActorDragDropGraphEdOp.h"
+#include "SceneOutlinerDragDrop.h"
+#include "DragAndDrop/ActorDragDropOp.h"
 #include "Widgets/Text/SInlineEditableTextBlock.h"
 
 #define LOCTEXT_NAMESPACE "LayersView"
@@ -123,55 +124,65 @@ bool SLayersViewRow::OnRenameLayerTextChanged(const FText& NewText, FText& OutEr
 
 void SLayersViewRow::OnDragLeave(const FDragDropEvent& DragDropEvent)
 {
-	TSharedPtr< FActorDragDropGraphEdOp > DragActorOp = DragDropEvent.GetOperationAs< FActorDragDropGraphEdOp >();
-	if (DragActorOp.IsValid())
+	TSharedPtr< FSceneOutlinerDragDropOp > DragOp = DragDropEvent.GetOperationAs< FSceneOutlinerDragDropOp >();
+	if (DragOp.IsValid())
 	{
-		DragActorOp->ResetToDefaultToolTip();
+		TSharedPtr< FActorDragDropOp > ActorDragOp = DragOp->GetSubOp<FActorDragDropOp>();
+		if (ActorDragOp.IsValid())
+		{
+			ActorDragOp->ResetToDefaultToolTip();
+		}
 	}
 }
 
 FReply SLayersViewRow::OnDragOver(const FGeometry& MyGeometry, const FDragDropEvent& DragDropEvent)
 {
-	TSharedPtr< FActorDragDropGraphEdOp > DragActorOp = DragDropEvent.GetOperationAs< FActorDragDropGraphEdOp >();
-	if (!DragActorOp.IsValid())
-	{
-		return FReply::Unhandled();
-	}
-
-	bool bCanAssign = false;
-	FText Message;
-	if (DragActorOp->Actors.Num() > 1)
-	{
-		bCanAssign = ViewModel->CanAssignActors(DragActorOp->Actors, OUT Message);
-	}
-	else
-	{
-		bCanAssign = ViewModel->CanAssignActor(DragActorOp->Actors[0], OUT Message);
-	}
-
-	if (bCanAssign)
-	{
-		DragActorOp->SetToolTip(FActorDragDropGraphEdOp::ToolTip_CompatibleGeneric, Message);
-	}
-	else
-	{
-		DragActorOp->SetToolTip(FActorDragDropGraphEdOp::ToolTip_IncompatibleGeneric, Message);
-	}
-
-	return FReply::Handled();
+	TSharedPtr< FSceneOutlinerDragDropOp > DragOp = DragDropEvent.GetOperationAs< FSceneOutlinerDragDropOp >();
+	if (DragOp.IsValid())
+	{
+		TSharedPtr< FActorDragDropOp > ActorDragOp = DragOp->GetSubOp<FActorDragDropOp>();
+		if (ActorDragOp.IsValid())
+		{
+			bool bCanAssign = false;
+			FText Message;
+			if (ActorDragOp->Actors.Num() > 1)
+			{
+				bCanAssign = ViewModel->CanAssignActors(ActorDragOp->Actors, OUT Message);
+			}
+			else
+			{
+				bCanAssign = ViewModel->CanAssignActor(ActorDragOp->Actors[0], OUT Message);
+			}
+
+			if (bCanAssign)
+			{
+				ActorDragOp->SetToolTip(Message, FEditorStyle::GetBrush(TEXT("Graph.ConnectorFeedback.OK")));
+			}
+			else
+			{
+				ActorDragOp->SetToolTip(Message, FEditorStyle::GetBrush(TEXT("Graph.ConnectorFeedback.Error")));
+			}
+
+			FReply::Handled();
+		}
+	}
+	return FReply::Unhandled();
 }
 
 FReply SLayersViewRow::OnDrop(const FGeometry& MyGeometry, const FDragDropEvent& DragDropEvent)
 {
-	TSharedPtr< FActorDragDropGraphEdOp > DragActorOp = DragDropEvent.GetOperationAs< FActorDragDropGraphEdOp >();
-	if (!DragActorOp.IsValid())
+	TSharedPtr< FSceneOutlinerDragDropOp > DragOp = DragDropEvent.GetOperationAs< FSceneOutlinerDragDropOp >();
+	if (DragOp.IsValid())
 	{
-		return FReply::Unhandled();
+		TSharedPtr< FActorDragDropOp > ActorDragOp = DragOp->GetSubOp<FActorDragDropOp>();
+		if (ActorDragOp.IsValid())
+		{
+			ViewModel->AddActors(ActorDragOp->Actors);
+			return FReply::Handled();
+		}
 	}
 
-	ViewModel->AddActors(DragActorOp->Actors);
-
-	return FReply::Handled();
+	return FReply::Unhandled();
 }
 
 FSlateColor SLayersViewRow::GetColorAndOpacity() const
@@ -184,12 +195,18 @@ FSlateColor SLayersViewRow::GetColorAndOpacity() const
 	bool bCanAcceptDrop = false;
 	TSharedPtr<FDragDropOperation> DragDropOp = FSlateApplication::Get().GetDragDroppingContent();
 
-	if (DragDropOp.IsValid() && DragDropOp->IsOfType<FActorDragDropGraphEdOp>())
-	{
-		TSharedPtr<FActorDragDropGraphEdOp> DragDropActorOp = StaticCastSharedPtr<FActorDragDropGraphEdOp>(DragDropOp);
-
-		FText Message;
-		bCanAcceptDrop = ViewModel->CanAssignActors(DragDropActorOp->Actors, OUT Message);
+	if (DragDropOp.IsValid() && DragDropOp->IsOfType<FSceneOutlinerDragDropOp>())
+	{
+		TSharedPtr<FSceneOutlinerDragDropOp> DragOp = StaticCastSharedPtr<FSceneOutlinerDragDropOp>(DragDropOp);
+		if (DragOp.IsValid())
+		{
+			TSharedPtr< FActorDragDropOp > ActorDragOp = DragOp->GetSubOp<FActorDragDropOp>();
+			if (ActorDragOp.IsValid())
+			{
+				FText Message;
+				bCanAcceptDrop = ViewModel->CanAssignActors(ActorDragOp->Actors, OUT Message);
+			}
+		}
 	}
 
 	return (bCanAcceptDrop) ? FSlateColor::UseForeground() : FLinearColor(0.30f, 0.30f, 0.30f);
diff --git a/Engine/Source/Editor/SceneOutliner/Private/SOutlinerTreeView.cpp b/Engine/Source/Editor/SceneOutliner/Private/SOutlinerTreeView.cpp
--- a/Engine/Source/Editor/SceneOutliner/Private/SOutlinerTreeView.cpp
+++ b/Engine/Source/Editor/SceneOutliner/Private/SOutlinerTreeView.cpp
@@ -72,6 +72,13 @@ FReply HandleDrop(TSharedPtr<SSceneOutliner> SceneOutlinerPtr, const FDragDropEv
 		return FReply::Unhandled();
 	}
 
+	// Don't handle this if we're not showing a hierarchy
+	const FSharedSceneOutlinerData& SharedData = SceneOutlinerPtr->GetSharedData();
+	if (!SharedData.bShowParentTree)
+	{
+		return FReply::Unhandled();
+	}
+
 	// Don't handle this if the scene outliner is not in a mode which supports drag and drop
 	if (!SceneOutlinerPtr->CanSupportDragAndDrop())
 	{
@@ -128,7 +135,10 @@ FReply SSceneOutlinerTreeView::OnDragOver(const FGeometry& MyGeometry, const FDr
 	FFolderTreeItem DropTarget(NAME_None);
 
 	auto Reply = HandleDropFromWeak(SceneOutlinerWeak, DragDropEvent, DropTarget, ValidationInfo);
-	UpdateOperationDecorator(DragDropEvent, ValidationInfo);
+	if (Reply.IsEventHandled())
+	{
+		UpdateOperationDecorator(DragDropEvent, ValidationInfo);
+	}
 
 	return Reply;
 }
@@ -170,8 +180,11 @@ void SSceneOutlinerTreeRow::OnDragEnter( const FGeometry& MyGeometry, const FDra
 	{
 		FSceneOutlinerDragValidationInfo ValidationInfo = FSceneOutlinerDragValidationInfo::Invalid();
 
-		HandleDrop(SceneOutlinerPtr, DragDropEvent, *ItemPtr, ValidationInfo, false);
-		UpdateOperationDecorator(DragDropEvent, ValidationInfo);
+		FReply Reply = HandleDrop(SceneOutlinerPtr, DragDropEvent, *ItemPtr, ValidationInfo, false);
+		if (Reply.IsEventHandled())
+		{
+			UpdateOperationDecorator(DragDropEvent, ValidationInfo);
+		}
 	}
 }
 
diff --git a/Engine/Source/Editor/UnrealEd/Public/DragAndDrop/ActorDragDropOp.h b/Engine/Source/Editor/UnrealEd/Public/DragAndDrop/ActorDragDropOp.h
--- a/Engine/Source/Editor/UnrealEd/Public/DragAndDrop/ActorDragDropOp.h
+++ b/Engine/Source/Editor/UnrealEd/Public/DragAndDrop/ActorDragDropOp.h
@@ -43,5 +43,7 @@ public:
 		{
 			CurrentHoverText = FText::Format(NSLOCTEXT("FActorDragDropOp", "FormatActors", "{0} Actors"), FText::AsNumber(InActors.Num()));
 		}
+
+		SetupDefaults();
 	}
 };
diff --git a/Engine/Source/Editor/UnrealEd/Public/DragAndDrop/FolderDragDropOp.h b/Engine/Source/Editor/UnrealEd/Public/DragAndDrop/FolderDragDropOp.h
--- a/Engine/Source/Editor/UnrealEd/Public/DragAndDrop/FolderDragDropOp.h
+++ b/Engine/Source/Editor/UnrealEd/Public/DragAndDrop/FolderDragDropOp.h
@@ -28,5 +28,7 @@ struct FFolderDragDropOp : public FDecoratedDragDropOp
 		{
 			CurrentHoverText = FText::Format(NSLOCTEXT("FFolderDragDropOp", "FormatFolders", "{0} Folders"), FText::AsNumber(Folders.Num()));
 		}
+
+		SetupDefaults();
 	}
 };
\ No newline at end of file
