CompositeDragDropOp refactor

Refactored the CompositeDragDropOp so that it can be interacted with transparently as any one of its sub ops.
Also now supports Folder Drops in the layer browser

diff --git a/Engine/Plugins/MovieScene/TemplateSequence/Source/TemplateSequenceEditor/Private/TemplateSequenceEditorModule.cpp b/Engine/Plugins/MovieScene/TemplateSequence/Source/TemplateSequenceEditor/Private/TemplateSequenceEditorModule.cpp
--- a/Engine/Plugins/MovieScene/TemplateSequence/Source/TemplateSequenceEditor/Private/TemplateSequenceEditorModule.cpp
+++ b/Engine/Plugins/MovieScene/TemplateSequence/Source/TemplateSequenceEditor/Private/TemplateSequenceEditorModule.cpp
@@ -5,7 +5,7 @@
 #include "AssetToolsModule.h"
 #include "CameraAnimationSequence.h"
 #include "ClassViewerModule.h"
-#include "DragAndDrop/ActorDragDropGraphEdOp.h"
+#include "DragAndDrop/ActorDragDropOp.h"
 #include "DragAndDrop/AssetDragDropOp.h"
 #include "DragAndDrop/ClassDragDropOp.h"
 #include "Framework/MultiBox/MultiBoxBuilder.h"
@@ -40,7 +40,7 @@ private:
 	bool OnSequencerReceivedDragOver(const FGeometry& MyGeometry, const FDragDropEvent& DragDropEvent, FReply& OutReply);
 	ESequencerDropResult OnSequencerAssetsDrop(const TArray<UObject*>& Assets, const FAssetDragDropOp& DragDropOp);
 	ESequencerDropResult OnSequencerClassesDrop(const TArray<TWeakObjectPtr<UClass>>& Classes, const FClassDragDropOp& DragDropOp);
-	ESequencerDropResult OnSequencerActorsDrop(const TArray<TWeakObjectPtr<AActor>>& Actors, const FActorDragDropGraphEdOp& DragDropOp);
+	ESequencerDropResult OnSequencerActorsDrop(const TArray<TWeakObjectPtr<AActor>>& Actors, const FActorDragDropOp& DragDropOp);
 
 	FText GetBoundActorClassName() const;
 	TSharedRef<SWidget> GetBoundActorClassMenuContent();
@@ -276,9 +276,9 @@ bool FTemplateSequenceCustomization::OnSequencerReceivedDragOver(const FGeometry
 
 	TSharedPtr<FDragDropOperation> Operation = DragDropEvent.GetOperation();
 	if (Operation.IsValid() && (
-		(Operation->IsOfType<FAssetDragDropOp>() && StaticCastSharedPtr<FAssetDragDropOp>(Operation)->GetAssetPaths().Num() <= 1) ||
-		(Operation->IsOfType<FClassDragDropOp>() && StaticCastSharedPtr<FClassDragDropOp>(Operation)->ClassesToDrop.Num() <= 1) ||
-		(Operation->IsOfType<FActorDragDropGraphEdOp>() && StaticCastSharedPtr<FActorDragDropGraphEdOp>(Operation)->Actors.Num() <= 1)))
+		(Operation->IsOfType<FAssetDragDropOp>() && Operation->CastTo<FAssetDragDropOp>()->GetAssetPaths().Num() <= 1) ||
+		(Operation->IsOfType<FClassDragDropOp>() && Operation->CastTo<FClassDragDropOp>()->ClassesToDrop.Num() <= 1) ||
+		(Operation->IsOfType<FActorDragDropOp>() && Operation->CastTo<FActorDragDropOp>()->Actors.Num() <= 1)))
 	{
 		bIsDragSupported = true;
 	}
@@ -318,7 +318,7 @@ ESequencerDropResult FTemplateSequenceCustomization::OnSequencerClassesDrop(cons
 	return ESequencerDropResult::Unhandled;
 }
 
-ESequencerDropResult FTemplateSequenceCustomization::OnSequencerActorsDrop(const TArray<TWeakObjectPtr<AActor>>& Actors, const FActorDragDropGraphEdOp& DragDropOp)
+ESequencerDropResult FTemplateSequenceCustomization::OnSequencerActorsDrop(const TArray<TWeakObjectPtr<AActor>>& Actors, const FActorDragDropOp& DragDropOp)
 {
 	return ESequencerDropResult::Unhandled;
 }
@@ -376,7 +376,7 @@ void FCameraAnimationSequenceCustomation::RegisterSequencerCustomization(FSequen
 
 	Customization.OnAssetsDrop.BindLambda([](const TArray<UObject*>&, const FAssetDragDropOp&) -> ESequencerDropResult { return ESequencerDropResult::DropDenied; });
 	Customization.OnClassesDrop.BindLambda([](const TArray<TWeakObjectPtr<UClass>>&, const FClassDragDropOp&) -> ESequencerDropResult { return ESequencerDropResult::DropDenied; });
-	Customization.OnActorsDrop.BindLambda([](const TArray<TWeakObjectPtr<AActor>>&, const FActorDragDropGraphEdOp&) -> ESequencerDropResult { return ESequencerDropResult::DropDenied; });
+	Customization.OnActorsDrop.BindLambda([](const TArray<TWeakObjectPtr<AActor>>&, const FActorDragDropOp&) -> ESequencerDropResult { return ESequencerDropResult::DropDenied; });
 
 	Builder.AddCustomization(Customization);
 }
diff --git a/Engine/Plugins/VirtualProduction/Takes/Source/TakeRecorderSources/Private/TakeRecorderSourcesModule.cpp b/Engine/Plugins/VirtualProduction/Takes/Source/TakeRecorderSources/Private/TakeRecorderSourcesModule.cpp
--- a/Engine/Plugins/VirtualProduction/Takes/Source/TakeRecorderSources/Private/TakeRecorderSourcesModule.cpp
+++ b/Engine/Plugins/VirtualProduction/Takes/Source/TakeRecorderSources/Private/TakeRecorderSourcesModule.cpp
@@ -14,7 +14,6 @@
 #include "TakeRecorderSources.h"
 #include "TakeRecorderSettings.h"
 #include "Features/IModularFeatures.h"
-#include "DragAndDrop/CompositeDragDropOp.h"
 #include "DragAndDrop/ActorDragDropOp.h"
 #include "DragAndDrop/FolderDragDropOp.h"
 #include "EngineUtils.h"
@@ -34,6 +33,7 @@
 #include "SceneOutlinerPublicTypes.h"
 #include "Subsystems/AssetEditorSubsystem.h"
 #include "ActorTreeItem.h"
+#include "EditorActorFolders.h"
 
 #include "TakeRecorderMicrophoneAudioSource.h"
 #include "TakeRecorderWorldSource.h"
@@ -176,45 +176,39 @@ struct FActorTakeRecorderDropHandler : ITakeRecorderDropHandler
 
 	virtual bool CanHandleOperation(TSharedPtr<FDragDropOperation> InOperation, UTakeRecorderSources* Sources) override
 	{
+		bool bCanHandle = false;
 		if (InOperation)
 		{
 			if (InOperation->IsOfType<FActorDragDropOp>())
 			{
-				return StaticCastSharedPtr<FActorDragDropOp>(InOperation)->Actors.Num() > 0;
+				bCanHandle |= InOperation->CastTo<FActorDragDropOp>()->Actors.Num() > 0;
 			}
-			else if (InOperation->IsOfType<FCompositeDragDropOp>())
+			if (InOperation->IsOfType<FFolderDragDropOp>())
 			{
-				return true;
+				bCanHandle |= true;
 			}
 		}
 
-		return false;
+		return bCanHandle;
 	}
 
 	TArray<AActor*> GetValidDropActors(TSharedPtr<FDragDropOperation> InOperation, UTakeRecorderSources* Sources)
 	{
-		FActorDragDropOp*  ActorDrag = nullptr;
-		FFolderDragDropOp* FolderDrag = nullptr;
+		TSharedPtr<FActorDragDropOp>  ActorDrag = nullptr;
+		TSharedPtr<FFolderDragDropOp> FolderDrag = nullptr;
 
 		FDragDropOperation* OperationPtr = InOperation.Get();
 		if (!OperationPtr)
 		{
 			return TArray<AActor*>();
 		}
-		else if (OperationPtr->IsOfType<FCompositeDragDropOp>())
-		{
-			FCompositeDragDropOp* OutlinerOp = static_cast<FCompositeDragDropOp*>(OperationPtr);
-			FolderDrag = OutlinerOp->GetSubOp<FFolderDragDropOp>().Get();
-			ActorDrag  = OutlinerOp->GetSubOp<FActorDragDropOp>().Get();
-
-		}
-		else if (OperationPtr->IsOfType<FActorDragDropOp>())
+		if (OperationPtr->IsOfType<FActorDragDropOp>())
 		{
-			ActorDrag = static_cast<FActorDragDropOp*>(OperationPtr);
+			ActorDrag = OperationPtr->CastTo<FActorDragDropOp>();
 		}
-		else if (OperationPtr->IsOfType<FFolderDragDropOp>())
+		if (OperationPtr->IsOfType<FFolderDragDropOp>())
 		{
-			FolderDrag = static_cast<FFolderDragDropOp*>(OperationPtr);
+			FolderDrag = OperationPtr->CastTo<FFolderDragDropOp>();
 		}
 
 		TArray<AActor*> DraggedActors;
@@ -233,21 +227,7 @@ struct FActorTakeRecorderDropHandler : ITakeRecorderDropHandler
 
 		if (FolderDrag)
 		{
-			// Copy the array onto the stack if it's within a reasonable size
-			TArray<FName, TInlineAllocator<16>> DraggedFolders(FolderDrag->Folders);
-
-			// Find any actors in the global editor world that have any of the dragged paths.
-			// WARNING: Actor iteration can be very slow, so this needs to be optimized
-			for (FActorIterator ActorIt(GWorld); ActorIt; ++ActorIt)
-			{
-				FName ActorPath = ActorIt->GetFolderPath();
-				if (ActorPath.IsNone() || !DraggedFolders.Contains(ActorPath))
-				{
-					continue;
-				}
-
-				DraggedActors.Add(*ActorIt);
-			}
+			FActorFolders::GetActorsFromFolders(*GWorld, FolderDrag->Folders, DraggedActors);
 		}
 
 		TArray<AActor*> ExistingActors;
diff --git a/Engine/Source/Editor/Layers/Private/SLayerBrowser.h b/Engine/Source/Editor/Layers/Private/SLayerBrowser.h
--- a/Engine/Source/Editor/Layers/Private/SLayerBrowser.h
+++ b/Engine/Source/Editor/Layers/Private/SLayerBrowser.h
@@ -16,10 +16,11 @@
 #include "Editor/Layers/Private/ActorsAssignedToSpecificLayersFilter.h"
 #include "ISceneOutlinerColumn.h"
 #include "Editor/Layers/Private/SceneOutlinerLayerContentsColumn.h"
-#include "SceneOutlinerDragDrop.h"
 #include "DragAndDrop/ActorDragDropOp.h"
+#include "DragAndDrop/FolderDragDropOp.h"
 #include "Editor/Layers/Private/SLayersView.h"
 #include "Editor/Layers/Private/SLayersCommandsMenu.h"
+#include "EditorActorFolders.h"
 
 class ISceneOutliner;
 
@@ -85,14 +86,10 @@ protected:
 	 */
 	virtual void OnDragLeave( const FDragDropEvent& DragDropEvent ) override
 	{
-		TSharedPtr< FSceneOutlinerDragDropOp > DragOp = DragDropEvent.GetOperationAs< FSceneOutlinerDragDropOp >();
+		TSharedPtr< FDecoratedDragDropOp > DragOp = DragDropEvent.GetOperationAs<FDecoratedDragDropOp>();
 		if (DragOp.IsValid())
 		{
-			TSharedPtr< FActorDragDropOp > ActorDragOp = DragOp->GetSubOp<FActorDragDropOp>();
-			if (ActorDragOp.IsValid())
-			{
-				ActorDragOp->ResetToDefaultToolTip();
-			}
+			DragOp->ResetToDefaultToolTip();
 		}
 	}
 
@@ -111,49 +108,67 @@ protected:
 			return FReply::Unhandled();
 		}
 
-		TSharedPtr< FSceneOutlinerDragDropOp > DragOp = DragDropEvent.GetOperationAs< FSceneOutlinerDragDropOp >();
-		if (DragOp.IsValid())
+		TArray<TWeakObjectPtr<AActor>> Actors;
+		TSharedPtr<FActorDragDropOp> ActorDragOp = DragDropEvent.GetOperationAs<FActorDragDropOp>();
+		TSharedPtr<FFolderDragDropOp> FolderDragOp = DragDropEvent.GetOperationAs<FFolderDragDropOp>();
+		if (ActorDragOp.IsValid() || FolderDragOp.IsValid())
 		{
-			TSharedPtr< FActorDragDropOp > ActorDragOp = DragOp->GetSubOp<FActorDragDropOp>();
-			if (ActorDragOp.IsValid())
-			{		
-				const FGeometry LayerContentsHeaderGeometry = SWidget::FindChildGeometry( MyGeometry, LayerContentsHeader.ToSharedRef() );
-				bool bValidDrop = LayerContentsHeaderGeometry.IsUnderLocation( DragDropEvent.GetScreenSpacePosition() );
+			const FGeometry LayerContentsHeaderGeometry = SWidget::FindChildGeometry( MyGeometry, LayerContentsHeader.ToSharedRef() );
+			bool bValidDrop = LayerContentsHeaderGeometry.IsUnderLocation( DragDropEvent.GetScreenSpacePosition() );
 
-				if( !bValidDrop && Mode == ELayerBrowserMode::LayerContents )
-				{
-					const FGeometry LayerContentsSectionGeometry = SWidget::FindChildGeometry( MyGeometry, LayerContentsSection.ToSharedRef() );
-					bValidDrop = LayerContentsSectionGeometry.IsUnderLocation( DragDropEvent.GetScreenSpacePosition() );
-				}
+			if( !bValidDrop && Mode == ELayerBrowserMode::LayerContents )
+			{
+				const FGeometry LayerContentsSectionGeometry = SWidget::FindChildGeometry( MyGeometry, LayerContentsSection.ToSharedRef() );
+				bValidDrop = LayerContentsSectionGeometry.IsUnderLocation( DragDropEvent.GetScreenSpacePosition() );
+			}
 
-				if( !bValidDrop )
-				{
-					return FReply::Unhandled();
-				}
+			if( bValidDrop && ActorDragOp.IsValid() && ActorDragOp->Actors.Num() > 0)
+			{
+				Actors = ActorDragOp->Actors;
+			}
 
-				if ( !ActorDragOp.IsValid() || ActorDragOp->Actors.Num() == 0 )
+			if (bValidDrop && FolderDragOp.IsValid())
+			{
+				if (UWorld* World = FolderDragOp->World.Get())
 				{
-					return FReply::Unhandled();
+					FActorFolders::GetWeakActorsFromFolders(*World, FolderDragOp->Folders, Actors);
 				}
+			}
+		}
 
-				bool bCanAssign = false;
-				FText Message;
-				if(ActorDragOp->Actors.Num() > 1 )
+		if (Actors.Num() > 0)
+		{
+			bool bCanAssign = false;
+			FText Message;
+			if (Actors.Num() > 1)
+			{
+				bCanAssign = SelectedLayerViewModel->CanAssignActors(Actors, OUT Message);
+			}
+			else
+			{
+				bCanAssign = SelectedLayerViewModel->CanAssignActor(Actors[0], OUT Message);
+			}
+
+			if (bCanAssign)
+			{
+				if (ActorDragOp.IsValid())
 				{
-					bCanAssign = SelectedLayerViewModel->CanAssignActors(ActorDragOp->Actors, OUT Message );
+					ActorDragOp->SetToolTip(Message, FEditorStyle::GetBrush(TEXT("Graph.ConnectorFeedback.OK")));
 				}
-				else
+				if (FolderDragOp.IsValid())
 				{
-					bCanAssign = SelectedLayerViewModel->CanAssignActor(ActorDragOp->Actors[ 0 ], OUT Message );
+					FolderDragOp->SetToolTip(Message, FEditorStyle::GetBrush(TEXT("Graph.ConnectorFeedback.OK")));
 				}
-
-				if ( bCanAssign )
+			}
+			else
+			{
+				if (ActorDragOp.IsValid())
 				{
-					ActorDragOp->SetToolTip(Message, FEditorStyle::GetBrush(TEXT("Graph.ConnectorFeedback.OK")));
+					ActorDragOp->SetToolTip(Message, FEditorStyle::GetBrush(TEXT("Graph.ConnectorFeedback.Error")));
 				}
-				else
+				if (FolderDragOp.IsValid())
 				{
-					ActorDragOp->SetToolTip(Message, FEditorStyle::GetBrush(TEXT("Graph.ConnectorFeedback.Error")));
+					FolderDragOp->SetToolTip(Message, FEditorStyle::GetBrush(TEXT("Graph.ConnectorFeedback.Error")));
 				}
 			}
 		}
@@ -177,32 +192,43 @@ protected:
 			return FReply::Unhandled();
 		}
 
-		TSharedPtr< FSceneOutlinerDragDropOp > DragOp = DragDropEvent.GetOperationAs< FSceneOutlinerDragDropOp >();
-		if (DragOp.IsValid())
+		bool bHandled = false;
+		TArray<TWeakObjectPtr<AActor>> ActorsToDrop;
+		TSharedPtr<FActorDragDropOp> ActorDragOp = DragDropEvent.GetOperationAs<FActorDragDropOp>();
+		if (ActorDragOp.IsValid())
 		{
-			TSharedPtr< FActorDragDropOp > ActorDragOp = DragOp->GetSubOp<FActorDragDropOp>();
-			if (ActorDragOp.IsValid())
-			{
-				const FGeometry LayerContentsHeaderGeometry = SWidget::FindChildGeometry(MyGeometry, LayerContentsHeader.ToSharedRef());
-				bool bValidDrop = LayerContentsHeaderGeometry.IsUnderLocation(DragDropEvent.GetScreenSpacePosition());
+			const FGeometry LayerContentsHeaderGeometry = SWidget::FindChildGeometry(MyGeometry, LayerContentsHeader.ToSharedRef());
+			bool bValidDrop = LayerContentsHeaderGeometry.IsUnderLocation(DragDropEvent.GetScreenSpacePosition());
 
-				if (!bValidDrop && Mode == ELayerBrowserMode::LayerContents)
-				{
-					const FGeometry LayerContentsSectionGeometry = SWidget::FindChildGeometry(MyGeometry, LayerContentsSection.ToSharedRef());
-					bValidDrop = LayerContentsSectionGeometry.IsUnderLocation(DragDropEvent.GetScreenSpacePosition());
-				}
+			if (!bValidDrop && Mode == ELayerBrowserMode::LayerContents)
+			{
+				const FGeometry LayerContentsSectionGeometry = SWidget::FindChildGeometry(MyGeometry, LayerContentsSection.ToSharedRef());
+				bValidDrop = LayerContentsSectionGeometry.IsUnderLocation(DragDropEvent.GetScreenSpacePosition());
+			}
 
-				if (!bValidDrop)
-				{
-					return FReply::Unhandled();
-				}
+			if (bValidDrop)
+			{
+				ActorsToDrop = ActorDragOp->Actors;
+				bHandled = true;
+			}
+		}
 
-				SelectedLayerViewModel->AddActors(ActorDragOp->Actors);
-				return FReply::Handled();
+		TSharedPtr<FFolderDragDropOp> FolderDragOp = DragDropEvent.GetOperationAs<FFolderDragDropOp>();
+		if (FolderDragOp.IsValid())
+		{
+			if (UWorld* World = FolderDragOp->World.Get())
+			{
+				FActorFolders::GetWeakActorsFromFolders(*World, FolderDragOp->Folders, ActorsToDrop);
+				bHandled = true;
 			}
 		}
 
-		return FReply::Unhandled();
+		if (ActorsToDrop.Num() > 0)
+		{
+			SelectedLayerViewModel->AddActors(ActorsToDrop);
+		}
+
+		return bHandled ? FReply::Handled() : FReply::Unhandled();
 	}
 
 
diff --git a/Engine/Source/Editor/Layers/Private/SLayersView.h b/Engine/Source/Editor/Layers/Private/SLayersView.h
--- a/Engine/Source/Editor/Layers/Private/SLayersView.h
+++ b/Engine/Source/Editor/Layers/Private/SLayersView.h
@@ -10,7 +10,6 @@
 #include "Widgets/SBoxPanel.h"
 #include "LayerCollectionViewCommands.h"
 #include "Widgets/Views/SHeaderRow.h"
-#include "SceneOutlinerDragDrop.h"
 #include "DragAndDrop/ActorDragDropOp.h"
 #include "Widgets/Views/STableViewBase.h"
 #include "Widgets/Views/STableRow.h"
@@ -18,6 +17,7 @@
 #include "Editor/Layers/Private/LayerCollectionViewModel.h"
 #include "Editor/Layers/Private/SLayersViewRow.h"
 #include "LayersDragDropOp.h"
+#include "EditorActorFolders.h"
 
 #define LOCTEXT_NAMESPACE "LayersView"
 
@@ -174,14 +174,10 @@ protected:
 	 */
 	virtual void OnDragLeave( const FDragDropEvent& DragDropEvent ) override
 	{
-		TSharedPtr< FSceneOutlinerDragDropOp > DragOp = DragDropEvent.GetOperationAs< FSceneOutlinerDragDropOp >();
+		TSharedPtr< FDecoratedDragDropOp > DragOp = DragDropEvent.GetOperationAs<FDecoratedDragDropOp>();
 		if (DragOp.IsValid())
 		{
-			TSharedPtr< FActorDragDropOp > ActorDragOp = DragOp->GetSubOp<FActorDragDropOp>();
-			if (ActorDragOp.IsValid())
-			{
-				ActorDragOp->ResetToDefaultToolTip();
-			}
+			DragOp->ResetToDefaultToolTip();
 		}
 	}
 
@@ -195,14 +191,11 @@ protected:
 	 */
 	virtual FReply OnDragOver( const FGeometry& MyGeometry, const FDragDropEvent& DragDropEvent ) override
 	{
-		TSharedPtr< FSceneOutlinerDragDropOp > DragOp = DragDropEvent.GetOperationAs< FSceneOutlinerDragDropOp >();
-		if (DragOp.IsValid())
+		TSharedPtr< FDecoratedDragDropOp > DecoratedDragDropOp = DragDropEvent.GetOperationAs<FDecoratedDragDropOp>();
+		
+		if (DecoratedDragDropOp.IsValid())
 		{
-			TSharedPtr< FActorDragDropOp > ActorDragOp = DragOp->GetSubOp<FActorDragDropOp>();
-			if (ActorDragOp.IsValid())
-			{
-				ActorDragOp->SetToolTip(LOCTEXT("OnDragOver", "Add Actors to New Layer"), FEditorStyle::GetBrush(TEXT("Graph.ConnectorFeedback.OK")));
-			}
+			DecoratedDragDropOp->SetToolTip(LOCTEXT("OnDragOver", "Add Actors to New Layer"), FEditorStyle::GetBrush(TEXT("Graph.ConnectorFeedback.OK")));
 		}
 
 		// We leave the event unhandled so the children of the ListView get a chance to grab the drag/drop
@@ -219,18 +212,33 @@ protected:
 	 */
 	virtual FReply OnDrop( const FGeometry& MyGeometry, const FDragDropEvent& DragDropEvent ) override
 	{
-		TSharedPtr< FSceneOutlinerDragDropOp > DragOp = DragDropEvent.GetOperationAs< FSceneOutlinerDragDropOp >();
-		if ( DragOp.IsValid() )
+		bool bHandled = false;
+		TArray<TWeakObjectPtr<AActor>> ActorsToDrop;
+
+		TSharedPtr<FActorDragDropOp> ActorDragOp = DragDropEvent.GetOperationAs<FActorDragDropOp>();
+		if (ActorDragOp.IsValid())
+		{
+			ActorsToDrop = ActorDragOp->Actors;
+			bHandled = true;
+		}
+
+		TSharedPtr<FFolderDragDropOp> FolderDragOp = DragDropEvent.GetOperationAs<FFolderDragDropOp>();
+		if (FolderDragOp.IsValid())
 		{
-			TSharedPtr< FActorDragDropOp > ActorDragOp = DragOp->GetSubOp<FActorDragDropOp>();
-			if ( ActorDragOp.IsValid() )
+			if (UWorld* World = FolderDragOp->World.Get())
 			{
-				ViewModel->AddActorsToNewLayer(ActorDragOp->Actors);
-				return FReply::Handled();
+				FActorFolders::GetWeakActorsFromFolders(*World, FolderDragOp->Folders, ActorsToDrop);
+
+				bHandled = true;
 			}
 		}
 
-		return FReply::Unhandled();
+		if (ActorsToDrop.Num() > 0)
+		{
+			ViewModel->AddActorsToNewLayer(ActorsToDrop);
+		}
+
+		return bHandled ? FReply::Handled() : FReply::Unhandled();
 	}
 
 	FReply OnDragRow(const FGeometry& MyGeometry, const FPointerEvent& MouseEvent)
diff --git a/Engine/Source/Editor/Layers/Private/SLayersViewRow.cpp b/Engine/Source/Editor/Layers/Private/SLayersViewRow.cpp
--- a/Engine/Source/Editor/Layers/Private/SLayersViewRow.cpp
+++ b/Engine/Source/Editor/Layers/Private/SLayersViewRow.cpp
@@ -6,9 +6,10 @@
 #include "Widgets/Input/SButton.h"
 #include "Widgets/Views/SListView.h"
 #include "EditorStyleSet.h"
-#include "SceneOutlinerDragDrop.h"
 #include "DragAndDrop/ActorDragDropOp.h"
+#include "DragAndDrop/FolderDragDropOp.h"
 #include "Widgets/Text/SInlineEditableTextBlock.h"
+#include "EditorActorFolders.h"
 
 #define LOCTEXT_NAMESPACE "LayersView"
 
@@ -124,65 +125,100 @@ bool SLayersViewRow::OnRenameLayerTextChanged(const FText& NewText, FText& OutEr
 
 void SLayersViewRow::OnDragLeave(const FDragDropEvent& DragDropEvent)
 {
-	TSharedPtr< FSceneOutlinerDragDropOp > DragOp = DragDropEvent.GetOperationAs< FSceneOutlinerDragDropOp >();
+	TSharedPtr< FDecoratedDragDropOp > DragOp = DragDropEvent.GetOperationAs<FDecoratedDragDropOp>();
 	if (DragOp.IsValid())
 	{
-		TSharedPtr< FActorDragDropOp > ActorDragOp = DragOp->GetSubOp<FActorDragDropOp>();
-		if (ActorDragOp.IsValid())
-		{
-			ActorDragOp->ResetToDefaultToolTip();
-		}
+		DragOp->ResetToDefaultToolTip();
 	}
 }
 
 FReply SLayersViewRow::OnDragOver(const FGeometry& MyGeometry, const FDragDropEvent& DragDropEvent)
 {
-	TSharedPtr< FSceneOutlinerDragDropOp > DragOp = DragDropEvent.GetOperationAs< FSceneOutlinerDragDropOp >();
-	if (DragOp.IsValid())
+	TArray<TWeakObjectPtr<AActor>> Actors;
+	TSharedPtr<FActorDragDropOp> ActorDragOp = DragDropEvent.GetOperationAs<FActorDragDropOp>();
+	if (ActorDragOp.IsValid() && ActorDragOp->Actors.Num() > 0)
 	{
-		TSharedPtr< FActorDragDropOp > ActorDragOp = DragOp->GetSubOp<FActorDragDropOp>();
-		if (ActorDragOp.IsValid())
+		Actors = ActorDragOp->Actors;
+	}
+
+	TSharedPtr<FFolderDragDropOp> FolderDragOp = DragDropEvent.GetOperationAs<FFolderDragDropOp>();
+	if (FolderDragOp.IsValid())
+	{
+		if (UWorld* World = FolderDragOp->World.Get())
 		{
-			bool bCanAssign = false;
-			FText Message;
-			if (ActorDragOp->Actors.Num() > 1)
+			FActorFolders::GetWeakActorsFromFolders(*World, FolderDragOp->Folders, Actors);
+		}
+	}
+
+	if (Actors.Num() > 0)
+	{
+		bool bCanAssign = false;
+		FText Message;
+		if (Actors.Num() > 1)
+		{
+			bCanAssign = ViewModel->CanAssignActors(Actors, OUT Message);
+		}
+		else
+		{
+			bCanAssign = ViewModel->CanAssignActor(Actors[0], OUT Message);
+		}
+
+		if (bCanAssign)
+		{
+			if (ActorDragOp.IsValid())
 			{
-				bCanAssign = ViewModel->CanAssignActors(ActorDragOp->Actors, OUT Message);
+				ActorDragOp->SetToolTip(Message, FEditorStyle::GetBrush(TEXT("Graph.ConnectorFeedback.OK")));
 			}
-			else
+			if (FolderDragOp.IsValid())
 			{
-				bCanAssign = ViewModel->CanAssignActor(ActorDragOp->Actors[0], OUT Message);
+				FolderDragOp->SetToolTip(Message, FEditorStyle::GetBrush(TEXT("Graph.ConnectorFeedback.OK")));
 			}
-
-			if (bCanAssign)
+		}
+		else
+		{
+			if (ActorDragOp.IsValid())
 			{
-				ActorDragOp->SetToolTip(Message, FEditorStyle::GetBrush(TEXT("Graph.ConnectorFeedback.OK")));
+				ActorDragOp->SetToolTip(Message, FEditorStyle::GetBrush(TEXT("Graph.ConnectorFeedback.Error")));
 			}
-			else
+			if (FolderDragOp.IsValid())
 			{
-				ActorDragOp->SetToolTip(Message, FEditorStyle::GetBrush(TEXT("Graph.ConnectorFeedback.Error")));
+				FolderDragOp->SetToolTip(Message, FEditorStyle::GetBrush(TEXT("Graph.ConnectorFeedback.Error")));
 			}
-
-			FReply::Handled();
 		}
+		return FReply::Handled();
 	}
+
 	return FReply::Unhandled();
 }
 
 FReply SLayersViewRow::OnDrop(const FGeometry& MyGeometry, const FDragDropEvent& DragDropEvent)
 {
-	TSharedPtr< FSceneOutlinerDragDropOp > DragOp = DragDropEvent.GetOperationAs< FSceneOutlinerDragDropOp >();
-	if (DragOp.IsValid())
+	bool bHandled = false;
+	TArray<TWeakObjectPtr<AActor>> ActorsToDrop;
+
+	TSharedPtr< FActorDragDropOp > ActorDragOp = DragDropEvent.GetOperationAs<FActorDragDropOp>();
+	if (ActorDragOp.IsValid())
 	{
-		TSharedPtr< FActorDragDropOp > ActorDragOp = DragOp->GetSubOp<FActorDragDropOp>();
-		if (ActorDragOp.IsValid())
+		ActorsToDrop = ActorDragOp->Actors;
+		bHandled = true;
+	}
+
+	TSharedPtr<FFolderDragDropOp> FolderDragOp = DragDropEvent.GetOperationAs<FFolderDragDropOp>();
+	if (FolderDragOp.IsValid())
+	{
+		if (UWorld* World = FolderDragOp->World.Get())
 		{
-			ViewModel->AddActors(ActorDragOp->Actors);
-			return FReply::Handled();
+			FActorFolders::GetWeakActorsFromFolders(*World, FolderDragOp->Folders, ActorsToDrop);
+			bHandled = true;
 		}
 	}
 
-	return FReply::Unhandled();
+	if (ActorsToDrop.Num() > 0)
+	{
+		ViewModel->AddActors(ActorsToDrop);
+	}
+
+	return bHandled ? FReply::Handled() : FReply::Unhandled();
 }
 
 FSlateColor SLayersViewRow::GetColorAndOpacity() const
@@ -195,18 +231,31 @@ FSlateColor SLayersViewRow::GetColorAndOpacity() const
 	bool bCanAcceptDrop = false;
 	TSharedPtr<FDragDropOperation> DragDropOp = FSlateApplication::Get().GetDragDroppingContent();
 
-	if (DragDropOp.IsValid() && DragDropOp->IsOfType<FSceneOutlinerDragDropOp>())
+	if (DragDropOp.IsValid() && (DragDropOp->IsOfType<FActorDragDropOp>() || DragDropOp->IsOfType<FFolderDragDropOp>()))
 	{
-		TSharedPtr<FSceneOutlinerDragDropOp> DragOp = StaticCastSharedPtr<FSceneOutlinerDragDropOp>(DragDropOp);
-		if (DragOp.IsValid())
+		TArray<TWeakObjectPtr<AActor>> DraggedActors;
+
+		TSharedPtr<FActorDragDropOp> ActorDragOp = DragDropOp->CastTo<FActorDragDropOp>();
+		TSharedPtr<FFolderDragDropOp> FolderDragOp = DragDropOp->CastTo<FFolderDragDropOp>();
+
+		if (ActorDragOp.IsValid())
 		{
-			TSharedPtr< FActorDragDropOp > ActorDragOp = DragOp->GetSubOp<FActorDragDropOp>();
-			if (ActorDragOp.IsValid())
+			DraggedActors = ActorDragOp->Actors;
+		}
+		if (FolderDragOp.IsValid())
+		{
+			auto World = FolderDragOp->World;
+			if (UWorld* WorldPtr = World.Get())
 			{
-				FText Message;
-				bCanAcceptDrop = ViewModel->CanAssignActors(ActorDragOp->Actors, OUT Message);
+				FActorFolders::GetWeakActorsFromFolders(*WorldPtr, FolderDragOp->Folders, DraggedActors);
 			}
 		}
+
+		if (DraggedActors.Num() > 0)
+		{
+			FText Message;
+			bCanAcceptDrop = ViewModel->CanAssignActors(DraggedActors, OUT Message);
+		}
 	}
 
 	return (bCanAcceptDrop) ? FSlateColor::UseForeground() : FLinearColor(0.30f, 0.30f, 0.30f);
diff --git a/Engine/Source/Editor/SceneOutliner/Private/ActorBrowsingMode.cpp b/Engine/Source/Editor/SceneOutliner/Private/ActorBrowsingMode.cpp
--- a/Engine/Source/Editor/SceneOutliner/Private/ActorBrowsingMode.cpp
+++ b/Engine/Source/Editor/SceneOutliner/Private/ActorBrowsingMode.cpp
@@ -731,7 +731,7 @@ TSharedPtr<FDragDropOperation> FActorBrowsingMode::CreateDragDropOperation(const
 	if (DraggedObjects.Has<FFolderTreeItem>())
 	{
 		TSharedPtr<FFolderDragDropOp> FolderOperation = MakeShareable(new FFolderDragDropOp);
-		FolderOperation->Init(DraggedObjects.GetData<FName>(SceneOutliner::FFolderPathSelector()));
+		FolderOperation->Init(DraggedObjects.GetData<FName>(SceneOutliner::FFolderPathSelector()), RepresentingWorld.Get());
 		OutlinerOp->AddSubOp(FolderOperation);
 	}
 	OutlinerOp->Construct();
diff --git a/Engine/Source/Editor/SceneOutliner/Private/ActorMode.cpp b/Engine/Source/Editor/SceneOutliner/Private/ActorMode.cpp
--- a/Engine/Source/Editor/SceneOutliner/Private/ActorMode.cpp
+++ b/Engine/Source/Editor/SceneOutliner/Private/ActorMode.cpp
@@ -24,12 +24,15 @@ namespace SceneOutliner
 {
 	bool FWeakActorSelector::operator()(const TWeakPtr<ISceneOutlinerTreeItem>& Item, TWeakObjectPtr<AActor>& DataOut) const
 	{
-		if (FActorTreeItem* ActorItem = Item.Pin()->CastTo<FActorTreeItem>())
+		if (TSharedPtr<ISceneOutlinerTreeItem> ItemPtr = Item.Pin())
 		{
-			if (ActorItem->IsValid())
+			if (FActorTreeItem* ActorItem = ItemPtr->CastTo<FActorTreeItem>())
 			{
-				DataOut = ActorItem->Actor;
-				return true;
+				if (ActorItem->IsValid())
+				{
+					DataOut = ActorItem->Actor;
+					return true;
+				}
 			}
 		}
 		return false;
@@ -37,28 +40,31 @@ namespace SceneOutliner
 
 	bool FActorSelector::operator()(const TWeakPtr<ISceneOutlinerTreeItem>& Item, AActor*& ActorPtrOut) const
 	{
-		if (FActorTreeItem* ActorItem = Item.Pin()->CastTo<FActorTreeItem>())
+		if (TSharedPtr<ISceneOutlinerTreeItem> ItemPtr = Item.Pin())
 		{
-			if (ActorItem->IsValid())
+			if (FActorTreeItem* ActorItem = ItemPtr->CastTo<FActorTreeItem>())
 			{
-				AActor* Actor = ActorItem->Actor.Get();
-				if (Actor)
+				if (ActorItem->IsValid())
 				{
-					ActorPtrOut = Actor;
-					return true;
+					AActor* Actor = ActorItem->Actor.Get();
+					if (Actor)
+					{
+						ActorPtrOut = Actor;
+						return true;
+					}
 				}
 			}
-		}
-		// If a component is selected, we meant for the owning actor to be selected
-		else if (FComponentTreeItem* ComponentItem = Item.Pin()->CastTo<FComponentTreeItem>())
-		{
-			if (ComponentItem->IsValid())
+			// If a component is selected, we meant for the owning actor to be selected
+			else if (FComponentTreeItem* ComponentItem = ItemPtr->CastTo<FComponentTreeItem>())
 			{
-				AActor* Actor = ComponentItem->Component->GetOwner();
-				if (Actor)
+				if (ComponentItem->IsValid())
 				{
-					ActorPtrOut = Actor;
-					return true;
+					AActor* Actor = ComponentItem->Component->GetOwner();
+					if (Actor)
+					{
+						ActorPtrOut = Actor;
+						return true;
+					}
 				}
 			}
 		}
diff --git a/Engine/Source/Editor/Sequencer/Private/SSequencer.cpp b/Engine/Source/Editor/Sequencer/Private/SSequencer.cpp
--- a/Engine/Source/Editor/Sequencer/Private/SSequencer.cpp
+++ b/Engine/Source/Editor/Sequencer/Private/SSequencer.cpp
@@ -52,7 +52,6 @@
 #include "DragAndDrop/ActorDragDropOp.h"
 #include "DragAndDrop/ActorDragDropGraphEdOp.h"
 #include "DragAndDrop/ClassDragDropOp.h"
-#include "DragAndDrop/CompositeDragDropOp.h"
 #include "DragAndDrop/FolderDragDropOp.h"
 #include "Widgets/Input/SSearchBox.h"
 #include "SSequencerTreeView.h"
@@ -98,7 +97,7 @@
 #include "SequencerTrackFilters.h"
 #include "SequencerTrackFilterExtension.h"
 #include "SequencerCustomizationManager.h"
-#include "EngineUtils.h"
+#include "EditorActorFolders.h"
 
 #define LOCTEXT_NAMESPACE "Sequencer"
 
@@ -342,7 +341,7 @@ void SSequencer::Construct(const FArguments& InArgs, TSharedRef<FSequencer> InSe
 	RootCustomization.OnAssetsDrop = InArgs._OnAssetsDrop;
 	RootCustomization.OnClassesDrop = InArgs._OnClassesDrop;
 	RootCustomization.OnActorsDrop = InArgs._OnActorsDrop;
-	RootCustomization.OnCompositeDrop = InArgs._OnCompositeDrop;
+	RootCustomization.OnFoldersDrop = InArgs._OnFoldersDrop;
 
 	// Get the desired display format from the user's settings each time.
 	TAttribute<EFrameNumberDisplayFormats> GetDisplayFormatAttr = MakeAttributeLambda(
@@ -2893,8 +2892,8 @@ FReply SSequencer::OnDragOver( const FGeometry& MyGeometry, const FDragDropEvent
 	if (Operation.IsValid() && (
 		Operation->IsOfType<FAssetDragDropOp>() ||
 		Operation->IsOfType<FClassDragDropOp>() ||
-		Operation->IsOfType<FActorDragDropGraphEdOp>() ||
-		Operation->IsOfType<FCompositeDragDropOp>() ) )
+		Operation->IsOfType<FActorDragDropOp>() ||
+		Operation->IsOfType<FFolderDragDropOp>() ) )
 	{
 		bIsDragSupported = true;
 	}
@@ -2936,30 +2935,30 @@ FReply SSequencer::OnDrop( const FGeometry& MyGeometry, const FDragDropEvent& Dr
 	{
 		if ( Operation->IsOfType<FAssetDragDropOp>() )
 		{
-			const auto& DragDropOp = StaticCastSharedPtr<FAssetDragDropOp>( Operation );
+			const auto& DragDropOp = Operation->CastTo<FAssetDragDropOp>();
 
 			OnAssetsDropped( *DragDropOp );
 			bWasDropHandled = true;
 		}
-		else if( Operation->IsOfType<FClassDragDropOp>() )
+		if( Operation->IsOfType<FClassDragDropOp>() )
 		{
-			const auto& DragDropOp = StaticCastSharedPtr<FClassDragDropOp>( Operation );
+			const auto& DragDropOp = Operation->CastTo<FClassDragDropOp>();
 
 			OnClassesDropped( *DragDropOp );
 			bWasDropHandled = true;
 		}
-		else if( Operation->IsOfType<FActorDragDropGraphEdOp>() )
+		if( Operation->IsOfType<FActorDragDropOp>() )
 		{
-			const auto& DragDropOp = StaticCastSharedPtr<FActorDragDropGraphEdOp>( Operation );
+			const auto& DragDropOp = Operation->CastTo<FActorDragDropOp>();
 
 			OnActorsDropped( *DragDropOp );
 			bWasDropHandled = true;
 		}
-		else if (Operation->IsOfType<FCompositeDragDropOp>())
+		if (Operation->IsOfType<FFolderDragDropOp>())
 		{
-			const auto& DragDropOp = StaticCastSharedPtr<FCompositeDragDropOp>(Operation);
+			const auto& DragDropOp = Operation->CastTo<FFolderDragDropOp>();
 
-			OnCompositeDropped(*DragDropOp);
+			OnFolderDropped(*DragDropOp);
 			bWasDropHandled = true;
 		}
 	}
@@ -3149,7 +3148,7 @@ void SSequencer::OnClassesDropped( const FClassDragDropOp& DragDropOp )
 	}
 }
 
-void SSequencer::OnActorsDropped( FActorDragDropGraphEdOp& DragDropOp )
+void SSequencer::OnActorsDropped( FActorDragDropOp& DragDropOp )
 {
 	const FScopedTransaction Transaction(LOCTEXT("DropActors", "Drop Actors"));
 
@@ -3173,31 +3172,13 @@ void SSequencer::OnActorsDropped( FActorDragDropGraphEdOp& DragDropOp )
 	}
 }
 
-void SSequencer::OnCompositeDropped( FCompositeDragDropOp& DragDropOp )
+void SSequencer::OnFolderDropped( FFolderDragDropOp& DragDropOp )
 {
 	const FScopedTransaction Transaction(LOCTEXT("DropActors", "Drop Actors"));
 
 	ESequencerDropResult DropResult = ESequencerDropResult::Unhandled;
 
 	TArray<TWeakObjectPtr<AActor>> DraggedActors;
-	FFolderDragDropOp* FolderDrag = DragDropOp.GetSubOp<FFolderDragDropOp>().Get();
-	FActorDragDropOp* ActorDrag = DragDropOp.GetSubOp<FActorDragDropOp>().Get();
-	if (ActorDrag)
-	{
-		DraggedActors.Reserve(ActorDrag->Actors.Num());
-		for (TWeakObjectPtr<AActor> WeakActor : ActorDrag->Actors)
-		{
-			if (AActor* Actor = WeakActor.Get())
-			{
-				DraggedActors.Add(Actor);
-			}
-		}
-	}
-
-	if (FolderDrag)
-	{
-		// Copy the array onto the stack if it's within a reasonable size
-		TArray<FName, TInlineAllocator<16>> DraggedFolders(FolderDrag->Folders);
 
 		// Find any actors in the global editor world that have any of the dragged paths.
 		// WARNING: Actor iteration can be very slow, so this needs to be optimized
@@ -3206,24 +3187,14 @@ void SSequencer::OnCompositeDropped( FCompositeDragDropOp& DragDropOp )
 		UWorld* World = PlaybackContext ? PlaybackContext->GetWorld() : nullptr;
 		if (World)
 		{
-			for (FActorIterator ActorIt(World); ActorIt; ++ActorIt)
-			{
-				FName ActorPath = ActorIt->GetFolderPath();
-				if (ActorPath.IsNone() || !DraggedFolders.Contains(ActorPath))
-				{
-					continue;
+		FActorFolders::GetWeakActorsFromFolders(*World, DragDropOp.Folders, DraggedActors);
 				}
-
-				DraggedActors.Add(*ActorIt);
-			}		
-		}
-	}	
 	
-	for (FOnCompositeDrop Delegate : OnCompositeDrop)
+	for (FOnFoldersDrop Delegate : OnFoldersDrop)
 	{
 		if (Delegate.IsBound())
 		{
-			DropResult = Delegate.Execute(DraggedActors, DragDropOp);
+			DropResult = Delegate.Execute(DragDropOp.Folders, DragDropOp);
 			if (DropResult != ESequencerDropResult::Unhandled)
 			{
 				break;
@@ -3945,7 +3916,7 @@ void SSequencer::ApplySequencerCustomizations(const TArray<FSequencerCustomizati
 	OnAssetsDrop.Reset();
 	OnActorsDrop.Reset();
 	OnClassesDrop.Reset();
-	OnCompositeDrop.Reset();
+	OnFoldersDrop.Reset();
 
 	ApplySequencerCustomization(RootCustomization);
 	for (const FSequencerCustomizationInfo& Info : Customizations)
@@ -3987,9 +3958,9 @@ void SSequencer::ApplySequencerCustomization(const FSequencerCustomizationInfo&
 	{
 		OnClassesDrop.Add(Customization.OnClassesDrop);
 	}
-	if (Customization.OnCompositeDrop.IsBound())
+	if (Customization.OnFoldersDrop.IsBound())
 	{
-		OnCompositeDrop.Add(Customization.OnCompositeDrop);
+		OnFoldersDrop.Add(Customization.OnFoldersDrop);
 	}
 }
 
diff --git a/Engine/Source/Editor/Sequencer/Private/SSequencer.h b/Engine/Source/Editor/Sequencer/Private/SSequencer.h
--- a/Engine/Source/Editor/Sequencer/Private/SSequencer.h
+++ b/Engine/Source/Editor/Sequencer/Private/SSequencer.h
@@ -18,10 +18,10 @@
 #include "Widgets/Text/STextBlock.h"
 #include "Sequencer.h"
 
-class FActorDragDropGraphEdOp;
+class FActorDragDropOp;
+class FFolderDragDropOp;
 class FAssetDragDropOp;
 class FClassDragDropOp;
-class FCompositeDragDropOp;
 class FMovieSceneClipboard;
 class FSequencerTimeSliderController;
 class FVirtualTrackArea;
@@ -261,8 +261,8 @@ public:
 		/** Called when an actor is dropped on the sequencer. Not called if OnReceivedDrop is bound and returned true. */
 		SLATE_EVENT( FOnActorsDrop, OnActorsDrop )
 
-		/** Called when an actor is dropped on the sequencer as a composite. Not called if OnReceivedDrop is bound and returned true. */
-		SLATE_EVENT( FOnCompositeDrop, OnCompositeDrop )
+		/** Called when a folder is dropped onto the sequencer. Not called if OnReceivedDrop is bound and returned true. */
+		SLATE_EVENT(FOnFoldersDrop, OnFoldersDrop)
 
 		/** Extender to use for the add menu. */
 		SLATE_ARGUMENT( TSharedPtr<FExtender>, AddMenuExtender )
@@ -512,14 +512,14 @@ private:
 	 *
 	 * @param	DragDropOp	Information about the actor(s) that was dropped
 	 */
-	void OnActorsDropped(FActorDragDropGraphEdOp& DragDropOp); 
+	void OnActorsDropped(FActorDragDropOp& DragDropOp); 
 
 	/**
-	 * Called when one or more objects are dropped into the widget as a composite
+	 * Called when one or more folders are dropped into the widget
 	 *
 	 * @param	DragDropOp	Information about the objects(s) that was dropped
 	 */
-	void OnCompositeDropped(FCompositeDragDropOp& DragDropOp); 
+	void OnFolderDropped(FFolderDragDropOp& DragDropOp); 
 
 	/** Called when a breadcrumb is clicked on in the sequencer */
 	void OnCrumbClicked(const FSequencerBreadcrumb& Item);
@@ -731,8 +731,8 @@ private:
 	/** Called when an actor is dropped on the sequencer. */
 	TArray<FOnActorsDrop> OnActorsDrop;
 
-	/** Called when an actor is dropped on the sequencer as a composite. */
-	TArray<FOnCompositeDrop> OnCompositeDrop;
+	/** Called when a folder is dropped on the sequencer. */
+	TArray<FOnFoldersDrop> OnFoldersDrop;
 
 	/** Stores the callbacks and extenders provided to the constructor. */
 	FSequencerCustomizationInfo RootCustomization;
diff --git a/Engine/Source/Editor/Sequencer/Public/SequencerCustomizationManager.h b/Engine/Source/Editor/Sequencer/Public/SequencerCustomizationManager.h
--- a/Engine/Source/Editor/Sequencer/Public/SequencerCustomizationManager.h
+++ b/Engine/Source/Editor/Sequencer/Public/SequencerCustomizationManager.h
@@ -7,8 +7,8 @@
 
 class FAssetDragDropOp;
 class FClassDragDropOp;
-class FActorDragDropGraphEdOp;
-class FCompositeDragDropOp;
+class FActorDragDropOp;
+class FFolderDragDropOp;
 class UMovieSceneSequence;
 
 /** A delegate that gets executed when a drag/drop event happens on the sequencer. The return value determines if the event was handled by the bound delegate. */
@@ -28,10 +28,10 @@ DECLARE_DELEGATE_RetVal_TwoParams(ESequencerDropResult, FOnAssetsDrop, const TAr
 DECLARE_DELEGATE_RetVal_TwoParams(ESequencerDropResult, FOnClassesDrop, const TArray<TWeakObjectPtr<UClass>>&, const FClassDragDropOp&);
 
 /** A delegate that gets executed when an actor is dropped on the sequencer. The return value determines if the operation was handled by the bound delegate. */
-DECLARE_DELEGATE_RetVal_TwoParams(ESequencerDropResult, FOnActorsDrop, const TArray<TWeakObjectPtr<AActor>>&, const FActorDragDropGraphEdOp&);
+DECLARE_DELEGATE_RetVal_TwoParams(ESequencerDropResult, FOnActorsDrop, const TArray<TWeakObjectPtr<AActor>>&, const FActorDragDropOp&);
 
-/** A delegate that gets executed when an actor is dropped on the sequencer as a composite. The return value determines if the operation was handled by the bound delegate. */
-DECLARE_DELEGATE_RetVal_TwoParams(ESequencerDropResult, FOnCompositeDrop, const TArray<TWeakObjectPtr<AActor>>&, const FCompositeDragDropOp&);
+/** A delegate that gets executed when a folder is dropped on the sequencer. The return value determines if the operation was handled by the bound delegate. */
+DECLARE_DELEGATE_RetVal_TwoParams(ESequencerDropResult, FOnFoldersDrop, const TArray<FName>&, const FFolderDragDropOp&);
 
 /** Class for specifying customizations to apply to a sequence editor. */
 struct FSequencerCustomizationInfo
@@ -57,8 +57,8 @@ struct FSequencerCustomizationInfo
 	/** Called when an actor is dropped on the sequencer. Only if OnReceivedDrop hasn't handled the event. */
 	FOnActorsDrop OnActorsDrop;
 
-	/** Called when an actor is dropped on the sequencer as a composite. Only if OnReceivedDrop hasn't handled the event. */
-	FOnCompositeDrop OnCompositeDrop;
+	/** Called when a folder is dropped on the sequencer. Only if OnReceivedDrop hasn't handled the event. */
+	FOnFoldersDrop OnFoldersDrop;
 };
 
 /** Class to pass to ISequencerCustomization for building a customization. */
diff --git a/Engine/Source/Editor/UnrealEd/Private/EditorActorFolders.cpp b/Engine/Source/Editor/UnrealEd/Private/EditorActorFolders.cpp
--- a/Engine/Source/Editor/UnrealEd/Private/EditorActorFolders.cpp
+++ b/Engine/Source/Editor/UnrealEd/Private/EditorActorFolders.cpp
@@ -563,4 +563,39 @@ bool FActorFolders::AddFolderToWorld(UWorld& InWorld, FName Path)
 	return false;
 }
 
+void FActorFolders::ForEachActorInFolders(UWorld& World, const TArray<FName>& Paths, TFunctionRef<bool(AActor*)> Operation)
+{
+	for (FActorIterator ActorIt(&World); ActorIt; ++ActorIt)
+	{
+		FName ActorPath = ActorIt->GetFolderPath();
+		if (ActorPath.IsNone() || !Paths.Contains(ActorPath))
+		{
+			continue;
+		}
+
+		if (!Operation(*ActorIt))
+		{
+			return;
+		}
+	}
+}
+
+void FActorFolders::GetActorsFromFolders(UWorld& World, const TArray<FName>& Paths, TArray<AActor*>& OutActors)
+{
+	ForEachActorInFolders(World, Paths, [&OutActors](AActor* InActor)
+		{
+			OutActors.Add(InActor);
+			return true;
+		});
+}
+
+void FActorFolders::GetWeakActorsFromFolders(UWorld& World, const TArray<FName>& Paths, TArray<TWeakObjectPtr<AActor>>& OutActors)
+{
+	ForEachActorInFolders(World, Paths, [&OutActors](AActor* InActor)
+		{
+			OutActors.Add(InActor);
+			return true;
+		});
+}
+
 #undef LOCTEXT_NAMESPACE
diff --git a/Engine/Source/Editor/UnrealEd/Public/DragAndDrop/CompositeDragDropOp.h b/Engine/Source/Editor/UnrealEd/Public/DragAndDrop/CompositeDragDropOp.h
--- a/Engine/Source/Editor/UnrealEd/Public/DragAndDrop/CompositeDragDropOp.h
+++ b/Engine/Source/Editor/UnrealEd/Public/DragAndDrop/CompositeDragDropOp.h
@@ -9,12 +9,17 @@
 class FCompositeDragDropOp : public FDecoratedDragDropOp
 {
 public:
-	DRAG_DROP_OPERATOR_TYPE(FCompositeDragDropOp, FDecoratedDragDropOp)
+	static const FString& GetTypeId() { static FString Type = TEXT("FCompositeDragDropOp"); return Type; }
+	virtual bool IsOfTypeImpl(const FString& Type) const override 
+	{
+		return GetTypeId() == Type || FDecoratedDragDropOp::IsOfTypeImpl(Type) || GetSubOpPtr(Type) != nullptr; 
+	}
 
 	FCompositeDragDropOp() : FDecoratedDragDropOp() {}
 
 	void AddSubOp(const TSharedPtr<FDragDropOperation>& SubOp)
 	{
+		check(!SubOp->IsOfType<FCompositeDragDropOp>());
 		SubOps.Add(SubOp);
 	}
 
@@ -43,6 +48,50 @@ public:
 		}
 		return nullptr;
 	}
+
+	virtual TSharedPtr<FDragDropOperation> ConvertTo(const FString& TypeId) override
+	{
+		// Attempt to convert this before getting subops
+		if (FDecoratedDragDropOp::IsOfTypeImpl(TypeId))
+		{
+			return AsShared();
+		}
+		else
+		{
+			// Will be nullptr if failed to convert
+			TSharedPtr<FDragDropOperation> SubOpAs = GetSubOpPtr(TypeId);
+			return SubOpAs;
+		}
+	}
+
+	virtual void ResetToDefaultToolTip() override
+	{
+		FDecoratedDragDropOp::ResetToDefaultToolTip();
+
+		for (const auto& SubOp : SubOps)
+		{
+			if (SubOp->IsOfType<FDecoratedDragDropOp>())
+			{
+				auto DecoratedSubOp = SubOp->CastTo<FDecoratedDragDropOp>();
+				if (DecoratedSubOp.IsValid())
+				{
+					DecoratedSubOp->ResetToDefaultToolTip();
+				}
+			}
+		}
+	}
+private:
+	TSharedPtr<FDragDropOperation> GetSubOpPtr(const FString& TypeId) const
+	{
+		for (const auto& SubOp : SubOps)
+		{
+			if (SubOp->IsOfTypeImpl(TypeId))
+			{
+				return SubOp;
+			}
+		}
+		return nullptr;
+	}
 protected:
 	TArray<TSharedPtr<FDragDropOperation>> SubOps;
 };
\ No newline at end of file
diff --git a/Engine/Source/Editor/UnrealEd/Public/DragAndDrop/FolderDragDropOp.h b/Engine/Source/Editor/UnrealEd/Public/DragAndDrop/FolderDragDropOp.h
--- a/Engine/Source/Editor/UnrealEd/Public/DragAndDrop/FolderDragDropOp.h
+++ b/Engine/Source/Editor/UnrealEd/Public/DragAndDrop/FolderDragDropOp.h
@@ -8,16 +8,21 @@
 #include "EditorFolderUtils.h"
 
 /* A drag/drop operation when dragging actor folders */
-struct FFolderDragDropOp : public FDecoratedDragDropOp
+class FFolderDragDropOp : public FDecoratedDragDropOp
 {
+public:
 	DRAG_DROP_OPERATOR_TYPE(FFolderDragDropOp, FDecoratedDragDropOp)
 
 	/** Array of folders that we are dragging */
 	TArray<FName> Folders;
+	/** World to which the folders belong */
+	TWeakObjectPtr<UWorld> World;
 
-	void Init(TArray<FName> InFolders)
+	void Init(TArray<FName> InFolders, UWorld* InWorld)
 	{
+		check(InWorld != nullptr);
 		Folders = MoveTemp(InFolders);
+		World = InWorld;
 
 		CurrentIconBrush = FEditorStyle::Get().GetBrush(TEXT("SceneOutliner.FolderClosed"));
 		if (Folders.Num() == 1)
diff --git a/Engine/Source/Editor/UnrealEd/Public/EditorActorFolders.h b/Engine/Source/Editor/UnrealEd/Public/EditorActorFolders.h
--- a/Engine/Source/Editor/UnrealEd/Public/EditorActorFolders.h
+++ b/Engine/Source/Editor/UnrealEd/Public/EditorActorFolders.h
@@ -105,6 +105,14 @@ struct UNREALED_API FActorFolders : public FGCObject
 	/** Rename the specified path to a new name */
 	bool RenameFolderInWorld(UWorld& World, FName OldPath, FName NewPath);
 
+	/** Apply an operation to each actor in the given list of folders. Will stop when operation returns false. */
+	static void ForEachActorInFolders(UWorld& World, const TArray<FName>& Paths, TFunctionRef<bool(AActor*)> Operation);
+	
+	/** Get an array of actors from a list of folders */
+	static void GetActorsFromFolders(UWorld& World, const TArray<FName>& Paths, TArray<AActor*>& OutActors);
+
+	/** Get an array of weak actor pointers from a list of folders */
+	static void GetWeakActorsFromFolders(UWorld& World, const TArray<FName>& Paths, TArray<TWeakObjectPtr<AActor>>& OutActors);
 private:
 
 	/** Returns true if folders have been created for the specified world */
diff --git a/Engine/Source/Runtime/SlateCore/Public/Input/DragAndDrop.h b/Engine/Source/Runtime/SlateCore/Public/Input/DragAndDrop.h
--- a/Engine/Source/Runtime/SlateCore/Public/Input/DragAndDrop.h
+++ b/Engine/Source/Runtime/SlateCore/Public/Input/DragAndDrop.h
@@ -37,6 +37,17 @@ public:
 		return IsOfTypeImpl(TType::GetTypeId());
 	}
 
+	/** Cast this drag and drop operation to the specified template type */
+	template <class OperationType> TSharedPtr<OperationType> CastTo()
+	{
+		TSharedPtr<FDragDropOperation> Converted(this->ConvertTo(OperationType::GetTypeId()));
+		if (Converted != nullptr)
+		{
+			return StaticCastSharedPtr<OperationType>(Converted);
+		}
+		return TSharedPtr<OperationType>();
+	}
+
 	/**
 	* Checks if this drag and drop operation is affected by a given PointerEvent.
 	*
@@ -88,6 +99,13 @@ public:
 	 */
 	void SetCursorOverride( TOptional<EMouseCursor::Type> CursorType );
 
+	/**
+	 * Checks whether this drag and drop operation can cast safely to the specified type.
+	 */
+	virtual bool IsOfTypeImpl(const FString& Type) const
+	{
+		return false;
+	}
 protected:
 
 	/**
@@ -105,12 +123,9 @@ protected:
 	 */
 	void DestroyCursorDecoratorWindow();
 
-	/**
-	 * Checks whether this drag and drop operation can cast safely to the specified type.
-	 */
-	virtual bool IsOfTypeImpl( const FString& Type ) const
+	virtual TSharedPtr<FDragDropOperation> ConvertTo(const FString& TypeId)
 	{
-		return false;
+		return IsOfTypeImpl(TypeId) ? AsShared() : TSharedPtr<FDragDropOperation>();
 	}
 
 protected:
diff --git a/Engine/Source/Runtime/SlateCore/Public/Input/DragAndDrop.inl b/Engine/Source/Runtime/SlateCore/Public/Input/DragAndDrop.inl
--- a/Engine/Source/Runtime/SlateCore/Public/Input/DragAndDrop.inl
+++ b/Engine/Source/Runtime/SlateCore/Public/Input/DragAndDrop.inl
@@ -10,10 +10,10 @@ TSharedPtr<OperationType> FDragDropEvent::GetOperationAs() const
 {
 	if (Content.IsValid() && Content->IsOfType<OperationType>())
 	{
-		return StaticCastSharedPtr<OperationType>(Content);
+		return Content->CastTo<OperationType>();
 	}
 	else
 	{
 		return TSharedPtr<OperationType>();
 	}
-}
+}
\ No newline at end of file
