Layers:
- Remove all code for loading and unloading actors from the Layers dialog.
- Don't allow adding partitioned actors into layers until it is fully supported.
- Clear actor layers when converting to world partition until it is fully supported.

diff --git a/Engine/Source/Editor/Layers/Private/ActorLayerViewModel.cpp b/Engine/Source/Editor/Layers/Private/ActorLayerViewModel.cpp
--- a/Engine/Source/Editor/Layers/Private/ActorLayerViewModel.cpp
+++ b/Engine/Source/Editor/Layers/Private/ActorLayerViewModel.cpp
@@ -70,16 +70,6 @@ bool FActorLayerViewModel::IsVisible() const
 	return Layer->IsVisible();
 }
 
-bool FActorLayerViewModel::ShouldActorsBeLoaded() const
-{
-	if( !Layer.IsValid() )
-	{
-		return false;
-	}
-
-	return Layer->ShouldLoadActors();
-}
-
 
 void FActorLayerViewModel::OnLayersChanged( const ELayersAction::Type Action, const TWeakObjectPtr< ULayer >& ChangedLayer, const FName& ChangedProperty )
 {
diff --git a/Engine/Source/Editor/Layers/Private/ActorLayerViewModel.h b/Engine/Source/Editor/Layers/Private/ActorLayerViewModel.h
--- a/Engine/Source/Editor/Layers/Private/ActorLayerViewModel.h
+++ b/Engine/Source/Editor/Layers/Private/ActorLayerViewModel.h
@@ -47,9 +47,6 @@ public:
 	/**	@return Whether the Layer is visible */
 	bool IsVisible() const;
 
-	/**	@return Whether the actors in this layer should be loaded */
-	bool ShouldActorsBeLoaded() const;
-
 	/********************************************************************
 	 * EVENTS
 	 ********************************************************************/
diff --git a/Engine/Source/Editor/Layers/Private/LayerCollectionViewModel.cpp b/Engine/Source/Editor/Layers/Private/LayerCollectionViewModel.cpp
--- a/Engine/Source/Editor/Layers/Private/LayerCollectionViewModel.cpp
+++ b/Engine/Source/Editor/Layers/Private/LayerCollectionViewModel.cpp
@@ -106,15 +106,6 @@ void FLayerCollectionViewModel::BindCommands()
 		FCanExecuteAction::CreateSP( this, &FLayerCollectionViewModel::MakeAllLayersVisible_CanExecute ) );
 
 
-	ActionList.MapAction(Commands.ToggleSelectedLayersActorsLoading,
-			FExecuteAction::CreateSP(this, &FLayerCollectionViewModel::ToggleSelectedLayersActorsLoading_Executed),
-			FCanExecuteAction::CreateSP(this, &FLayerCollectionViewModel::ToggleSelectedLayersActorsLoading_CanExecute));
-
-	ActionList.MapAction(Commands.MakeAllLayersLoadActors,
-		FExecuteAction::CreateSP(this, &FLayerCollectionViewModel::MakeAllLayersLoadActors_Executed),
-		FCanExecuteAction::CreateSP(this, &FLayerCollectionViewModel::MakeAllLayersLoadActors_CanExecute));
-
-
 	ActionList.MapAction( Commands.RequestRenameLayer,
 		FExecuteAction::CreateSP( this, &FLayerCollectionViewModel::RequestRenameLayer_Executed ),
 		FCanExecuteAction::CreateSP( this, &FLayerCollectionViewModel::RequestRenameLayer_CanExecute ) );
@@ -280,7 +271,6 @@ void FLayerCollectionViewModel::OnLayerAdded( const TWeakObjectPtr< ULayer >& Ad
 
 	const TSharedRef< FLayerViewModel > NewLayerViewModel = FLayerViewModel::Create( AddedLayer, Editor );
 	NewLayerViewModel->OnVisibilityToggled().AddSP( this, &FLayerCollectionViewModel::ToggleLayerVisibility );
-	NewLayerViewModel->OnActorsLoadingToggled().AddSP( this, &FLayerCollectionViewModel::ToggleLayerActorsLoading );
 	AllLayerViewModels.Add( NewLayerViewModel );
 
 	// We specifically ignore filters when dealing with single additions
@@ -323,7 +313,6 @@ void FLayerCollectionViewModel::CreateViewModels( const TArray< TWeakObjectPtr<
 	{
 		const TSharedRef< FLayerViewModel > NewLayerViewModel = FLayerViewModel::Create( *LayerIt, Editor );
 		NewLayerViewModel->OnVisibilityToggled().AddSP( this, &FLayerCollectionViewModel::ToggleLayerVisibility );
-		NewLayerViewModel->OnActorsLoadingToggled().AddSP( this, &FLayerCollectionViewModel::ToggleLayerActorsLoading );
 		AllLayerViewModels.Add( NewLayerViewModel );
 
 		if( Filters->PassesAllFilters( NewLayerViewModel ) )
@@ -649,71 +638,6 @@ bool FLayerCollectionViewModel::MakeAllLayersVisible_CanExecute() const
 }
 
 
-// ToggleLayerActorsLoading ------------------------------------------
-void FLayerCollectionViewModel::ToggleLayerActorsLoading(const TSharedPtr<FLayerViewModel>& InLayer)
-{
-	if (SelectedLayers.Find(InLayer) == INDEX_NONE)
-	{
-		// Given layer wasn't selected so toggle only its own visibility
-		const FScopedTransaction Transaction(LOCTEXT("ToggleActorsLoading", "Toggle Layer Actors Loading"));
-		WorldLayers->ToggleLayerActorsLoading(InLayer->GetFName());
-	}
-	else
-	{
-		// Toggle actors loading of selected layers to the same state as the given layer
-		bool bShouldLoadActors = InLayer->ShouldLoadActors();
-
-		TArray< FName > SelectedLayerNames;
-		for (auto LayersIt = SelectedLayers.CreateConstIterator(); LayersIt; ++LayersIt)
-		{
-			const TSharedPtr< FLayerViewModel >& Layer = *LayersIt;
-			if (Layer->ShouldLoadActors() == bShouldLoadActors)
-			{
-				SelectedLayerNames.Add(Layer->GetFName());
-			}
-		}
-
-		const FScopedTransaction Transaction(LOCTEXT("ToggleSelectedLayersActorsLoading", "Toggle Layer Actors Loading"));
-		WorldLayers->ToggleLayersActorsLoading(SelectedLayerNames);
-	}
-}
-
-
-// ToggleSelectedLayersActorsLoading ------------------------------------------
-void FLayerCollectionViewModel::ToggleSelectedLayersActorsLoading_Executed()
-{
-	if (SelectedLayers.Num() == 0)
-	{
-		return;
-	}
-
-	TArray< FName > SelectedLayerNames;
-	AppendSelectLayerNames(SelectedLayerNames);
-
-	const FScopedTransaction Transaction(LOCTEXT("ToggleSelectedLayersActorsLoading", "Toggle Layers Actors Loading"));
-
-	WorldLayers->ToggleLayersActorsLoading(SelectedLayerNames);
-}
-
-bool FLayerCollectionViewModel::ToggleSelectedLayersActorsLoading_CanExecute() const
-{
-	return SelectedLayers.Num() > 0;
-}
-
-
-// MakeAllLayersLoadActors ----------------------------------------------------
-void FLayerCollectionViewModel::MakeAllLayersLoadActors_Executed()
-{
-	const FScopedTransaction Transaction(LOCTEXT("MakeAllLayersLoadActors", "Make All Layers Load Actors"));
-	WorldLayers->MakeAllLayersLoadActors();
-}
-
-bool FLayerCollectionViewModel::MakeAllLayersLoadActors_CanExecute() const
-{
-	return AllLayerViewModels.Num() > 0;
-}
-
-
 // RequestRenameLayer ---------------------------------------------------------
 void FLayerCollectionViewModel::RequestRenameLayer_Executed()
 {
diff --git a/Engine/Source/Editor/Layers/Private/LayerCollectionViewModel.h b/Engine/Source/Editor/Layers/Private/LayerCollectionViewModel.h
--- a/Engine/Source/Editor/Layers/Private/LayerCollectionViewModel.h
+++ b/Engine/Source/Editor/Layers/Private/LayerCollectionViewModel.h
@@ -97,9 +97,6 @@ public:
 	/** Toggles the visibility of the given layer and those of the other layers selected with it */
 	void ToggleLayerVisibility( const TSharedPtr< FLayerViewModel >& InLayer );
 
-	/** Toggles the actors loading of the given layer and those of the other layers selected with it */
-	void ToggleLayerActorsLoading(const TSharedPtr< FLayerViewModel >& InLayer);
-
 	/********************************************************************
 	 * EVENTS
 	 ********************************************************************/
@@ -230,18 +227,6 @@ private:
 	/** @return	whether all layers can be made visible */
 	bool MakeAllLayersVisible_CanExecute() const;
 
-	/** Toggles the selected layers actors loading */
-	void ToggleSelectedLayersActorsLoading_Executed();
-
-	/** @return	whether the selected layers actors loading can be toggled */
-	bool ToggleSelectedLayersActorsLoading_CanExecute() const;
-
-	/** Enable actors loading on all layers */
-	void MakeAllLayersLoadActors_Executed();
-
-	/** @return	whether it's possible to enable actors loading on all layers */
-	bool MakeAllLayersLoadActors_CanExecute() const;
-
 	/** Requests renaming of selected layers */
 	void RequestRenameLayer_Executed();
 
diff --git a/Engine/Source/Editor/Layers/Private/LayerViewModel.cpp b/Engine/Source/Editor/Layers/Private/LayerViewModel.cpp
--- a/Engine/Source/Editor/Layers/Private/LayerViewModel.cpp
+++ b/Engine/Source/Editor/Layers/Private/LayerViewModel.cpp
@@ -93,17 +93,6 @@ bool FLayerViewModel::IsVisible() const
 }
 
 
-bool FLayerViewModel::ShouldLoadActors() const
-{
-	if (!Layer.IsValid())
-	{
-		return false;
-	}
-
-	return Layer->ShouldLoadActors();
-}
-
-
 void FLayerViewModel::ToggleVisibility()
 {
 	if( !Layer.IsValid() )
diff --git a/Engine/Source/Editor/Layers/Private/LayerViewModel.h b/Engine/Source/Editor/Layers/Private/LayerViewModel.h
--- a/Engine/Source/Editor/Layers/Private/LayerViewModel.h
+++ b/Engine/Source/Editor/Layers/Private/LayerViewModel.h
@@ -51,9 +51,6 @@ public:
 	/**	@return Whether the Layer is visible */
 	bool IsVisible() const;
 
-	/**	@return Whether actors from this layer should be loaded */
-	bool ShouldLoadActors() const;
-
 	/** 
 	 *	Toggles the specified layer's visibility
 	 */
@@ -193,10 +190,6 @@ public:
 	DECLARE_EVENT_OneParam( FLayerViewModel, FVisibilityToggledEvent, const TSharedPtr<FLayerViewModel>& )
 	FVisibilityToggledEvent& OnVisibilityToggled() { return VisibilityToggledEvent; }
 
-	/** Broadcasts whenever a layer actors loading toggle is requested */
-	DECLARE_EVENT_OneParam(FLayerViewModel, ActorsLoadingToggledEvent, const TSharedPtr<FLayerViewModel>&)
-	FVisibilityToggledEvent& OnActorsLoadingToggled() { return ActorsLoadingToggledEvent; }
-
 private:
 
 	/**
diff --git a/Engine/Source/Editor/Layers/Private/SLayersView.h b/Engine/Source/Editor/Layers/Private/SLayersView.h
--- a/Engine/Source/Editor/Layers/Private/SLayersView.h
+++ b/Engine/Source/Editor/Layers/Private/SLayersView.h
@@ -82,11 +82,6 @@ public:
 			.DefaultLabel(NSLOCTEXT("LayersView", "Visibility", "Visibility"))
 			.FixedWidth(40.0f)
 
-			/** Layer actors loading column */
-			+SHeaderRow::Column(LayersView::ColumnID_ActorsLoading)
-			.DefaultLabel(NSLOCTEXT("LayersView", "ActorsLoading", "Actors Loading"))
-			.FixedWidth(40.0f)
-
 			/** LayerName label column */
 			+SHeaderRow::Column(LayersView::ColumnID_LayerLabel)
 			.DefaultLabel(LOCTEXT("Column_LayerNameLabel", "Layer"));
diff --git a/Engine/Source/Editor/Layers/Private/SLayersViewRow.cpp b/Engine/Source/Editor/Layers/Private/SLayersViewRow.cpp
--- a/Engine/Source/Editor/Layers/Private/SLayersViewRow.cpp
+++ b/Engine/Source/Editor/Layers/Private/SLayersViewRow.cpp
@@ -78,23 +78,6 @@ TSharedRef< SWidget > SLayersViewRow::GenerateWidgetForColumn(const FName& Colum
 			]
 		;
 	}
-	else if (ColumnID == LayersView::ColumnID_ActorsLoading)
-	{
-		TableRowContent =
-			SAssignNew(ActorsLoadingButton, SButton)
-			.ContentPadding(0)
-			.ButtonStyle(FEditorStyle::Get(), "NoBorder")
-			.OnClicked(this, &SLayersViewRow::OnToggleActorsLoading)
-			.ToolTipText(LOCTEXT("ActorsLoadingButtonToolTip", "Toggle Actors Loading"))
-			.HAlign(HAlign_Center)
-			.VAlign(VAlign_Center)
-			.Content()
-			[
-				SNew(SImage)
-				.Image(this, &SLayersViewRow::GetActorsLoadingBrushForLayer)
-			]
-		;
-	}
 	else
 	{
 		checkf(false, TEXT("Unknown ColumnID provided to SLayersView"));
@@ -275,16 +258,4 @@ const FSlateBrush* SLayersViewRow::GetVisibilityBrushForLayer() const
 	}
 }
 
-const FSlateBrush* SLayersViewRow::GetActorsLoadingBrushForLayer() const
-{
-	if (ViewModel->ShouldLoadActors())
-	{
-		return FEditorStyle::GetBrush("SessionBrowser.StatusRunning");
-	}
-	else
-	{
-		return FEditorStyle::GetBrush("SessionBrowser.StatusTimedOut");
-	}
-}
-
 #undef LOCTEXT_NAMESPACE
diff --git a/Engine/Source/Editor/Layers/Private/SLayersViewRow.h b/Engine/Source/Editor/Layers/Private/SLayersViewRow.h
--- a/Engine/Source/Editor/Layers/Private/SLayersViewRow.h
+++ b/Engine/Source/Editor/Layers/Private/SLayersViewRow.h
@@ -19,7 +19,6 @@ namespace LayersView
 	/** IDs for list columns */
 	static const FName ColumnID_LayerLabel( "Layer" );
 	static const FName ColumnID_Visibility( "Visibility" );
-	static const FName ColumnID_ActorsLoading( "ActorsLoading" );
 }
 
 /**
@@ -131,14 +130,6 @@ private:
 	 */
 	const FSlateBrush* GetVisibilityBrushForLayer() const;
 
-	/**
-	 *	Called to get the Slate Image Brush representing the actors loading state of
-	 *	the layer this row widget represents
-	 *
-	 *	@return	The SlateBrush representing the layer's actors loading state
-	 */
-	const FSlateBrush* GetActorsLoadingBrushForLayer() const;
-
 private:
 
 	/** The layer associated with this row of data */
@@ -147,9 +138,6 @@ private:
 	/**	The visibility button for the layer */
 	TSharedPtr< SButton > VisibilityButton;
 
-	/**	The actors loading button for the layer */
-	TSharedPtr< SButton > ActorsLoadingButton;
-
 	/** The string to highlight on any text contained in the row widget */
 	TAttribute< FText > HighlightText;
 
diff --git a/Engine/Source/Editor/UnrealEd/Private/Commandlets/WorldPartitionConvertCommandlet.cpp b/Engine/Source/Editor/UnrealEd/Private/Commandlets/WorldPartitionConvertCommandlet.cpp
--- a/Engine/Source/Editor/UnrealEd/Private/Commandlets/WorldPartitionConvertCommandlet.cpp
+++ b/Engine/Source/Editor/UnrealEd/Private/Commandlets/WorldPartitionConvertCommandlet.cpp
@@ -836,6 +836,9 @@ int32 UWorldPartitionConvertCommandlet::Main(const FString& Params)
 				{
 					Actor->GridPlacement = DefaultGridPlacement;
 				}
+
+				// Clear actor layers as they are not supported yet in world partition
+				Actor->Layers.Empty();
 			}
 		}
 
diff --git a/Engine/Source/Editor/UnrealEd/Private/Layers/LayersSubsystem.cpp b/Engine/Source/Editor/UnrealEd/Private/Layers/LayersSubsystem.cpp
--- a/Engine/Source/Editor/UnrealEd/Private/Layers/LayersSubsystem.cpp
+++ b/Engine/Source/Editor/UnrealEd/Private/Layers/LayersSubsystem.cpp
@@ -184,7 +184,8 @@ bool ULayersSubsystem::IsActorValidForLayer(AActor* Actor)
 	const bool bIsBuilderBrush = FActorEditorUtils::IsABuilderBrush(Actor);
 	const bool bIsHidden = (Actor->GetClass()->GetDefaultObject<AActor>()->bHiddenEd == true);
 	const bool bIsInEditorWorld = (Actor->GetWorld()->WorldType == EWorldType::Editor);
-	const bool bIsValid = !bIsHidden && !bIsBuilderBrush && bIsInEditorWorld;
+	const bool bIsPartitionedActor = Actor->GetLevel()->bIsPartitioned;
+	const bool bIsValid = !bIsHidden && !bIsBuilderBrush && bIsInEditorWorld && !bIsPartitionedActor;
 
 	return bIsValid;
 }
@@ -1145,77 +1146,6 @@ void ULayersSubsystem::MakeAllLayersVisible()
 }
 
 
-
-void ULayersSubsystem::SetLayerActorsLoading(const FName& LayerName, const bool bShouldLoadActors)
-{
-	SetLayersActorsLoading({ LayerName }, bShouldLoadActors);
-}
-
-
-void ULayersSubsystem::SetLayersActorsLoading(const TArray< FName >& LayerNames, const bool bShouldLoadActors)
-{
-	bool bChangeOccurred = false;
-	for (const auto& LayerName : LayerNames)
-	{
-		ULayer* Layer = EnsureLayerExists(LayerName);
-		check(Layer != NULL);
-
-		if (Layer->ShouldLoadActors() != bShouldLoadActors)
-		{
-			Layer->Modify();
-			Layer->SetShouldLoadActors(bShouldLoadActors);
-			LayersChanged.Broadcast(ELayersAction::Modify, Layer, "bShouldLoadActors");
-			bChangeOccurred = true;
-		}
-	}
-
-	if (bChangeOccurred)
-	{
-		GEditor->RedrawLevelEditingViewports();
-	}
-}
-
-
-void ULayersSubsystem::ToggleLayerActorsLoading(const FName& LayerName)
-{
-	ULayer* Layer = EnsureLayerExists(LayerName);
-	check(Layer != NULL);
-
-	Layer->Modify();
-	Layer->SetShouldLoadActors(!Layer->ShouldLoadActors());
-
-	LayersChanged.Broadcast(ELayersAction::Modify, Layer, "bShouldLoadActors");
-	GEditor->RedrawLevelEditingViewports();
-}
-
-
-void ULayersSubsystem::ToggleLayersActorsLoading(const TArray< FName >& LayerNames)
-{
-	for (const auto& LayerName : LayerNames)
-	{
-		ToggleLayerActorsLoading(LayerName);
-	}
-
-	GEditor->RedrawLevelEditingViewports();
-}
-
-
-void ULayersSubsystem::MakeAllLayersLoadActors()
-{
-	for (auto Layer : GetWorld()->Layers)
-	{
-		if (!Layer->ShouldLoadActors())
-		{
-			Layer->Modify();
-			Layer->SetShouldLoadActors(true);
-			LayersChanged.Broadcast(ELayersAction::Modify, Layer, "bShouldLoadActors");
-		}
-	}
-
-	GEditor->RedrawLevelEditingViewports();
-}
-
-
 ULayer* ULayersSubsystem::GetLayer(const FName& LayerName) const
 {
 	for (auto LayerIt = GetWorld()->Layers.CreateConstIterator(); LayerIt; ++LayerIt)
diff --git a/Engine/Source/Editor/UnrealEd/Public/Layers/LayersSubsystem.h b/Engine/Source/Editor/UnrealEd/Public/Layers/LayersSubsystem.h
--- a/Engine/Source/Editor/UnrealEd/Public/Layers/LayersSubsystem.h
+++ b/Engine/Source/Editor/UnrealEd/Public/Layers/LayersSubsystem.h
@@ -508,44 +508,6 @@ public:
 	UFUNCTION(BlueprintCallable, Category = Layers)
 	virtual void MakeAllLayersVisible() final;
 
-	/**
-	 * Changes the named layer's actors loading behavior to the provided state
-	 *
-	 * @param	LayerName	The name of the layer to affect.
-	 * @param	bLoadActors	If true, actors from this layer will be loaded.
-	 */
-	UFUNCTION(BlueprintCallable, Category = Layers)
-	void SetLayerActorsLoading(const FName& LayerName, const bool bLoadActors);
-	/**
-	 * Changes the actors loading behavior of the named layers to the provided state (world partition only)
-	 *
-	 * @param	LayerNames	The names of the layers to affect
-	 * @param	bLoadActors	If true actors from the provided layers will be loaded; false otherwise
-	 */
-	UFUNCTION(BlueprintCallable, Category = Layers)
-	void SetLayersActorsLoading(const TArray< FName >& LayerNames, const bool bLoadActors);
-
-	/**
-	 * Toggles the named layer's actors loading behavior
-	 *
-	 * @param LayerName	The name of the layer to affect
-	 */
-	UFUNCTION(BlueprintCallable, Category = Layers)
-	void ToggleLayerActorsLoading(const FName& LayerName);
-	/**
-	 * Toggles the actors loading behavior of all of the named layers
-	 *
-	 * @param	LayerNames	The names of the layers to affect
-	 */
-	UFUNCTION(BlueprintCallable, Category = Layers)
-	void ToggleLayersActorsLoading(const TArray< FName >& LayerNames);
-
-	/**
-	 * Enable actors loading for all layers
-	 */
-	UFUNCTION(BlueprintCallable, Category = Layers)
-	void MakeAllLayersLoadActors();
-
 	/**
 	 * Gets the ULayer Object of the named layer
 	 *
diff --git a/Engine/Source/Runtime/Engine/Classes/Layers/Layer.h b/Engine/Source/Runtime/Engine/Classes/Layers/Layer.h
--- a/Engine/Source/Runtime/Engine/Classes/Layers/Layer.h
+++ b/Engine/Source/Runtime/Engine/Classes/Layers/Layer.h
@@ -39,9 +39,6 @@ public:
 	void SetVisible(bool bIsVisible);
 	bool IsVisible() const;
 
-	void SetShouldLoadActors(bool bShouldLoadActors);
-	bool ShouldLoadActors() const;
-
 	const TArray<FLayerActorStats>& GetActorStats() const;
 	void ClearActorStats();
 	void AddToStats(AActor* Actor);
@@ -56,10 +53,6 @@ private:
 	UPROPERTY()
 	uint32 bIsVisible:1;
 
-	/** Whether actors associated with the layer should be loaded or not */
-	UPROPERTY()
-	uint32 bShouldLoadActors:1;
-
 	/** 
 	 * Basic stats regarding the number of Actors and their types currently assigned to the Layer 
 	 */
diff --git a/Engine/Source/Runtime/Engine/Private/Layers/Layer.cpp b/Engine/Source/Runtime/Engine/Private/Layers/Layer.cpp
--- a/Engine/Source/Runtime/Engine/Private/Layers/Layer.cpp
+++ b/Engine/Source/Runtime/Engine/Private/Layers/Layer.cpp
@@ -6,7 +6,6 @@ ULayer::ULayer(const FObjectInitializer& ObjectInitializer)
 	: Super(ObjectInitializer)
 	, LayerName( NAME_None )
 	, bIsVisible( true )
-	, bShouldLoadActors( true )
 {
 }
 
@@ -30,16 +29,6 @@ bool ULayer::IsVisible() const
 	return bIsVisible;
 }
 
-void ULayer::SetShouldLoadActors(bool bInShouldLoadActors)
-{
-	bShouldLoadActors = bInShouldLoadActors;
-}
-
-bool ULayer::ShouldLoadActors() const
-{
-	return bShouldLoadActors;
-}
-
 void ULayer::ClearActorStats()
 {
 	ActorStats.Empty();
diff --git a/Engine/Source/Runtime/Engine/Private/WorldPartition/WorldPartition.cpp b/Engine/Source/Runtime/Engine/Private/WorldPartition/WorldPartition.cpp
--- a/Engine/Source/Runtime/Engine/Private/WorldPartition/WorldPartition.cpp
+++ b/Engine/Source/Runtime/Engine/Private/WorldPartition/WorldPartition.cpp
@@ -421,23 +421,6 @@ void UWorldPartition::RegisterDelegates()
 
 		FEditorDelegates::PreBeginPIE.AddUObject(this, &UWorldPartition::OnPreBeginPIE);
 		FEditorDelegates::EndPIE.AddUObject(this, &UWorldPartition::OnEndPIE);
-
-		if (LayerSubSystem)
-		{
-			OnLayersChangedHandle = LayerSubSystem->OnLayersChanged().AddLambda([this](const ELayersAction::Type Action, const TWeakObjectPtr<ULayer>& ChangedLayer, const FName& ChangedProperty)
-			{
-				static FName NAME_bShouldLoadActors(TEXT("bShouldLoadActors"));
-				if (ChangedProperty == NAME_bShouldLoadActors)
-				{
-					RefreshLoadedCells();
-				}
-			});
-
-			OnActorsLayersChangedHandle = LayerSubSystem->OnActorsLayersChanged().AddLambda([this](const TWeakObjectPtr<AActor>& ChangedActor)
-			{
-				UpdateActorDesc(ChangedActor.Get());
-			});
-		}
 	}
 }
 
@@ -445,12 +428,6 @@ void UWorldPartition::UnregisterDelegates()
 {
 	if (GEditor && !IsTemplate())
 	{
-		if (LayerSubSystem)
-		{
-			LayerSubSystem->OnLayersChanged().Remove(OnLayersChangedHandle);
-			LayerSubSystem->OnActorsLayersChanged().Remove(OnActorsLayersChangedHandle);
-		}
-
 		FEditorDelegates::PreBeginPIE.RemoveAll(this);
 		FEditorDelegates::EndPIE.RemoveAll(this);
 
@@ -848,28 +825,6 @@ void UWorldPartition::UpdateLoadingEditorCell(UWorldPartitionEditorCell* Cell, b
 
 	Cell->Modify(false);
 
-	auto ShouldActorBeLoaded = [&](const FWorldPartitionActorDesc* ActorDesc)
-	{
-		if (LayerSubSystem)
-		{
-			uint32 NumValidLayers = 0;
-			for (const FName& LayerName : ActorDesc->GetLayers())
-			{
-				if (ULayer* Layer = LayerSubSystem->GetLayer(LayerName))
-				{
-					if (Layer->ShouldLoadActors())
-					{
-						return true;
-					}
-					NumValidLayers++;
-				}
-			}
-			return !NumValidLayers;
-		}
-
-		return true;
-	};
-
 	auto UnloadActor = [this](FWorldPartitionActorDesc* ActorDesc, AActor* Actor)
 	{
 		check(Actor);
@@ -908,40 +863,32 @@ void UWorldPartition::UpdateLoadingEditorCell(UWorldPartitionEditorCell* Cell, b
 			// or directly referenced by a loaded cell.
 			check(Actor || !ActorDesc->GetLoadedRefCount());
 
-			// Filter actor against visible layers and load it if required
-			if (ShouldActorBeLoaded(ActorDesc))
+			bool bIsAlreadyInLoadedActors = false;
+			Cell->LoadedActors.Add(ActorDesc, &bIsAlreadyInLoadedActors);
+
+			if (bIsAlreadyInLoadedActors)
 			{
-				bool bIsAlreadyInLoadedActors = false;
-				Cell->LoadedActors.Add(ActorDesc, &bIsAlreadyInLoadedActors);
+				// We already hold a reference to this actor
+				check(Actor);
+				check(ActorDesc->GetLoadedRefCount());
+				UE_LOG(LogWorldPartition, Verbose, TEXT(" ==> Skipped already loaded actor %s"), *Actor->GetFullName());
+			}
+			else
+			{
+				const uint32 ActorRefCount = ActorDesc->AddLoadedRefCount();
 
-				if (bIsAlreadyInLoadedActors)
+				if (ActorRefCount == 1)
 				{
-					// We already hold a reference to this actor
-					check(Actor);
-					check(ActorDesc->GetLoadedRefCount());
-					UE_LOG(LogWorldPartition, Verbose, TEXT(" ==> Skipped already loaded actor %s"), *Actor->GetFullName());
+					// Register actor
+					Actor = RegisterActor(ActorDesc);
 				}
 				else
 				{
-					const uint32 ActorRefCount = ActorDesc->AddLoadedRefCount();
-
-					if (ActorRefCount == 1)
-					{
-						// Register actor
-						Actor = RegisterActor(ActorDesc);
-					}
-					else
-					{
-						check(Actor);
-						UE_LOG(LogWorldPartition, Verbose, TEXT(" ==> Referenced unloaded actor %s(%d)"), *Actor->GetFullName(), ActorRefCount);
-					}
-
-					check(Actor && Actor->GetActorGuid() == ActorDesc->GetGuid());
+					check(Actor);
+					UE_LOG(LogWorldPartition, Verbose, TEXT(" ==> Referenced unloaded actor %s(%d)"), *Actor->GetFullName(), ActorRefCount);
 				}
-			}
-			else if (Cell->LoadedActors.Remove(ActorDesc))
-			{
-				UnloadActor(ActorDesc, Actor);
+
+				check(Actor && Actor->GetActorGuid() == ActorDesc->GetGuid());
 			}
 		}
 	}
