Removed remaining traces of Actor loading in Layers

diff --git a/Engine/Source/Editor/Layers/Private/LayerCollectionViewCommands.h b/Engine/Source/Editor/Layers/Private/LayerCollectionViewCommands.h
--- a/Engine/Source/Editor/Layers/Private/LayerCollectionViewCommands.h
+++ b/Engine/Source/Editor/Layers/Private/LayerCollectionViewCommands.h
@@ -46,9 +46,6 @@ public:
 		UI_COMMAND( ToggleSelectedLayersVisibility, "Toggle Visibility of Selected Layers", "Toggles the visibility of the selected layers in the viewports", EUserInterfaceActionType::Button, FInputChord() );
 		UI_COMMAND( MakeAllLayersVisible, "Make All Layers Visible", "Toggles all layers to a visible state in the viewports", EUserInterfaceActionType::Button, FInputChord() );
 
-		UI_COMMAND( ToggleSelectedLayersActorsLoading, "Toggle Actors Loading for Selected Layers", "Toggles the actors loading for the selected layers", EUserInterfaceActionType::Button, FInputChord());
-		UI_COMMAND( MakeAllLayersLoadActors, "Make All Layers Load Actors", "Enable Actors Loading on all Layers", EUserInterfaceActionType::Button, FInputChord());
-	
 		UI_COMMAND( RequestRenameLayer, "Rename", "Rename the selected layer.", EUserInterfaceActionType::Button, FInputChord( EKeys::F2 ) );	
 	}
 
@@ -81,12 +78,6 @@ public:
 	/** Makes all layers visible */
 	TSharedPtr< FUICommandInfo > MakeAllLayersVisible;
 
-	/** Toggles actor loading for the currently selected layers */
-	TSharedPtr< FUICommandInfo > ToggleSelectedLayersActorsLoading;
-
-	/** Makes all layers load their actors */
-	TSharedPtr< FUICommandInfo > MakeAllLayersLoadActors;
-
 	/** Requests a rename on the layer */
 	TSharedPtr< FUICommandInfo > RequestRenameLayer;
 };
diff --git a/Engine/Source/Editor/Layers/Private/SLayersCommandsMenu.cpp b/Engine/Source/Editor/Layers/Private/SLayersCommandsMenu.cpp
--- a/Engine/Source/Editor/Layers/Private/SLayersCommandsMenu.cpp
+++ b/Engine/Source/Editor/Layers/Private/SLayersCommandsMenu.cpp
@@ -65,12 +65,6 @@ void SLayersCommandsMenu::Construct(const FArguments& InArgs, const TSharedRef<
 			MenuBuilder.AddMenuEntry(Commands.MakeAllLayersVisible);
 		}
 		MenuBuilder.EndSection();
-
-		MenuBuilder.BeginSection("LayersActorsLoading", LOCTEXT("ActorsLoadingMenuHeader", "Actors Loading"));
-		{
-			MenuBuilder.AddMenuEntry(Commands.MakeAllLayersLoadActors);
-		}
-		MenuBuilder.EndSection();
 	}
 
 	ChildSlot
