SceneOutliner: Reverting composite drag drop transparent cast

diff --git a/Engine/Plugins/MovieScene/TemplateSequence/Source/TemplateSequenceEditor/Private/TemplateSequenceEditorModule.cpp b/Engine/Plugins/MovieScene/TemplateSequence/Source/TemplateSequenceEditor/Private/TemplateSequenceEditorModule.cpp
--- a/Engine/Plugins/MovieScene/TemplateSequence/Source/TemplateSequenceEditor/Private/TemplateSequenceEditorModule.cpp
+++ b/Engine/Plugins/MovieScene/TemplateSequence/Source/TemplateSequenceEditor/Private/TemplateSequenceEditorModule.cpp
@@ -325,9 +325,9 @@ bool FTemplateSequenceCustomization::OnSequencerReceivedDragOver(const FGeometry
 
 	TSharedPtr<FDragDropOperation> Operation = DragDropEvent.GetOperation();
 	if (Operation.IsValid() && (
-		(Operation->IsOfType<FAssetDragDropOp>() && Operation->CastTo<FAssetDragDropOp>()->GetAssetPaths().Num() <= 1) ||
-		(Operation->IsOfType<FClassDragDropOp>() && Operation->CastTo<FClassDragDropOp>()->ClassesToDrop.Num() <= 1) ||
-		(Operation->IsOfType<FActorDragDropOp>() && Operation->CastTo<FActorDragDropOp>()->Actors.Num() <= 1)))
+		(Operation->IsOfType<FAssetDragDropOp>() && StaticCastSharedPtr<FAssetDragDropOp>(Operation)->GetAssetPaths().Num() <= 1) ||
+		(Operation->IsOfType<FClassDragDropOp>() && StaticCastSharedPtr<FClassDragDropOp>(Operation)->ClassesToDrop.Num() <= 1) ||
+		(Operation->IsOfType<FActorDragDropOp>() && StaticCastSharedPtr<FActorDragDropOp>(Operation)->Actors.Num() <= 1)))
 	{
 		bIsDragSupported = true;
 	}
diff --git a/Engine/Plugins/VirtualProduction/Takes/Source/TakeRecorderSources/Private/TakeRecorderSourcesModule.cpp b/Engine/Plugins/VirtualProduction/Takes/Source/TakeRecorderSources/Private/TakeRecorderSourcesModule.cpp
--- a/Engine/Plugins/VirtualProduction/Takes/Source/TakeRecorderSources/Private/TakeRecorderSourcesModule.cpp
+++ b/Engine/Plugins/VirtualProduction/Takes/Source/TakeRecorderSources/Private/TakeRecorderSourcesModule.cpp
@@ -17,6 +17,7 @@
 #include "Features/IModularFeatures.h"
 #include "DragAndDrop/ActorDragDropOp.h"
 #include "DragAndDrop/FolderDragDropOp.h"
+#include "DragAndDrop/CompositeDragDropOp.h"
 #include "EngineUtils.h"
 #include "Algo/Sort.h"
 #include "ScopedTransaction.h"
@@ -183,18 +184,25 @@ struct FActorTakeRecorderDropHandler : ITakeRecorderDropHandler
 			TSharedPtr<FActorDragDropOp>  ActorDrag = nullptr;
 			TSharedPtr<FFolderDragDropOp> FolderDrag = nullptr;
 
-			FDragDropOperation* OperationPtr = InOperation.Get();
-			if (!OperationPtr)
+			if (!InOperation.IsValid())
 			{
 				return false;
 			}
 			if (InOperation->IsOfType<FActorDragDropOp>())
 			{
-				ActorDrag = OperationPtr->CastTo<FActorDragDropOp>();
+				ActorDrag = StaticCastSharedPtr<FActorDragDropOp>(InOperation);
 			}
-			if (InOperation->IsOfType<FFolderDragDropOp>())
+			else if (InOperation->IsOfType<FFolderDragDropOp>())
 			{
-				FolderDrag = OperationPtr->CastTo<FFolderDragDropOp>();
+				FolderDrag = StaticCastSharedPtr<FFolderDragDropOp>(InOperation);
+			}
+			else if (InOperation->IsOfType<FCompositeDragDropOp>())
+			{
+				if (const TSharedPtr<FCompositeDragDropOp> CompositeDrag = StaticCastSharedPtr<FCompositeDragDropOp>(InOperation))
+				{
+					ActorDrag = CompositeDrag->GetSubOp<FActorDragDropOp>();
+					FolderDrag = CompositeDrag->GetSubOp<FFolderDragDropOp>();
+				}
 			}
 
 			if (ActorDrag)
@@ -236,18 +244,25 @@ struct FActorTakeRecorderDropHandler : ITakeRecorderDropHandler
 		TSharedPtr<FActorDragDropOp>  ActorDrag = nullptr;
 		TSharedPtr<FFolderDragDropOp> FolderDrag = nullptr;
 
-		FDragDropOperation* OperationPtr = InOperation.Get();
-		if (!OperationPtr)
+		if (!InOperation.IsValid())
 		{
 			return TArray<AActor*>();
 		}
-		if (OperationPtr->IsOfType<FActorDragDropOp>())
+		if (InOperation->IsOfType<FActorDragDropOp>())
 		{
-			ActorDrag = OperationPtr->CastTo<FActorDragDropOp>();
+			ActorDrag = StaticCastSharedPtr<FActorDragDropOp>(InOperation);
 		}
-		if (OperationPtr->IsOfType<FFolderDragDropOp>())
+		else if (InOperation->IsOfType<FFolderDragDropOp>())
 		{
-			FolderDrag = OperationPtr->CastTo<FFolderDragDropOp>();
+			FolderDrag = StaticCastSharedPtr<FFolderDragDropOp>(InOperation);
+		}
+		else if (InOperation->IsOfType<FCompositeDragDropOp>())
+		{
+			if (const TSharedPtr<FCompositeDragDropOp> CompositeOp = StaticCastSharedPtr<FCompositeDragDropOp>(InOperation))
+			{
+				ActorDrag = CompositeOp->GetSubOp<FActorDragDropOp>();
+				FolderDrag = CompositeOp->GetSubOp<FFolderDragDropOp>();
+			}
 		}
 
 		TArray<AActor*> DraggedActors;
diff --git a/Engine/Source/Editor/DataLayerEditor/Private/DataLayer/DataLayerMode.cpp b/Engine/Source/Editor/DataLayerEditor/Private/DataLayer/DataLayerMode.cpp
--- a/Engine/Source/Editor/DataLayerEditor/Private/DataLayer/DataLayerMode.cpp
+++ b/Engine/Source/Editor/DataLayerEditor/Private/DataLayer/DataLayerMode.cpp
@@ -302,15 +302,27 @@ TArray<AActor*> FDataLayerMode::GetActorsFromOperation(const FDragDropOperation&
 		}
 	};
 
-	TSharedPtr<FActorDragDropOp> ActorDragOp = const_cast<FDragDropOperation&>(Operation).CastTo<FActorDragDropOp>();
-	if (ActorDragOp.IsValid())
+	if (Operation.IsOfType<FActorDragDropOp>())
 	{
-		GetActorsFromActorOperation(*ActorDragOp);
+		const FActorDragDropOp& ActorDragOp = StaticCast<const FActorDragDropOp&>(Operation);
+		GetActorsFromActorOperation(ActorDragOp);
 	}
-	TSharedPtr<FFolderDragDropOp> FolderDragOp = const_cast<FDragDropOperation&>(Operation).CastTo<FFolderDragDropOp>();
-	if (FolderDragOp.IsValid())
+	if (Operation.IsOfType<FFolderDragDropOp>())
 	{
-		GetActorsFromFolderOperation(*FolderDragOp);
+		const FFolderDragDropOp& FolderDragOp = StaticCast<const FFolderDragDropOp&>(Operation);
+		GetActorsFromFolderOperation(FolderDragOp);
+	}
+	if (Operation.IsOfType<FCompositeDragDropOp>())
+	{
+		const FCompositeDragDropOp& CompositeDragOp = StaticCast<const FCompositeDragDropOp&>(Operation);
+		if (TSharedPtr<const FActorDragDropOp> ActorSubOp = CompositeDragOp.GetSubOp<FActorDragDropOp>())
+		{
+			GetActorsFromActorOperation(*ActorSubOp);
+		}
+		if (TSharedPtr<const FFolderDragDropOp> FolderSubOp = CompositeDragOp.GetSubOp<FFolderDragDropOp>())
+		{
+			GetActorsFromFolderOperation(*FolderSubOp);
+		}
 	}
 	return Actors.Array();
 }
diff --git a/Engine/Source/Editor/DataLayerEditor/Private/DataLayer/DataLayerPropertyTypeCustomization.cpp b/Engine/Source/Editor/DataLayerEditor/Private/DataLayer/DataLayerPropertyTypeCustomization.cpp
--- a/Engine/Source/Editor/DataLayerEditor/Private/DataLayer/DataLayerPropertyTypeCustomization.cpp
+++ b/Engine/Source/Editor/DataLayerEditor/Private/DataLayer/DataLayerPropertyTypeCustomization.cpp
@@ -140,7 +140,7 @@ FReply FDataLayerPropertyTypeCustomization::OnDrop(TSharedPtr<FDragDropOperation
 {
 	if (InDragDrop.IsValid() && InDragDrop->IsOfType<FDataLayerDragDropOp>())
 	{
-		TSharedPtr<const FDataLayerDragDropOp> DataLayerDragDropOp = InDragDrop->CastTo<FDataLayerDragDropOp>();
+		TSharedPtr<const FDataLayerDragDropOp> DataLayerDragDropOp = StaticCastSharedPtr<const FDataLayerDragDropOp>(InDragDrop);
 		if (DataLayerDragDropOp.IsValid())
 		{
 			const TArray<FName>& DataLayerLabels = DataLayerDragDropOp->DataLayerLabels;
@@ -160,7 +160,7 @@ bool FDataLayerPropertyTypeCustomization::OnVerifyDrag(TSharedPtr<FDragDropOpera
 {
 	if (InDragDrop.IsValid() && InDragDrop->IsOfType<FDataLayerDragDropOp>())
 	{
-		TSharedPtr<const FDataLayerDragDropOp> DataLayerDragDropOp = InDragDrop->CastTo<FDataLayerDragDropOp>();
+		TSharedPtr<const FDataLayerDragDropOp> DataLayerDragDropOp = StaticCastSharedPtr<const FDataLayerDragDropOp>(InDragDrop);
 		return DataLayerDragDropOp.IsValid() && DataLayerDragDropOp->DataLayerLabels.Num() == 1;
 	}
 	return false;
diff --git a/Engine/Source/Editor/Layers/Private/SLayerBrowser.h b/Engine/Source/Editor/Layers/Private/SLayerBrowser.h
--- a/Engine/Source/Editor/Layers/Private/SLayerBrowser.h
+++ b/Engine/Source/Editor/Layers/Private/SLayerBrowser.h
@@ -18,6 +18,7 @@
 #include "Editor/Layers/Private/SceneOutlinerLayerContentsColumn.h"
 #include "DragAndDrop/ActorDragDropOp.h"
 #include "DragAndDrop/FolderDragDropOp.h"
+#include "DragAndDrop/CompositeDragDropOp.h"
 #include "Editor/Layers/Private/SLayersView.h"
 #include "Editor/Layers/Private/SLayersCommandsMenu.h"
 #include "EditorActorFolders.h"
@@ -109,8 +110,20 @@ protected:
 		}
 
 		TArray<TWeakObjectPtr<AActor>> Actors;
-		TSharedPtr<FActorDragDropOp> ActorDragOp = DragDropEvent.GetOperationAs<FActorDragDropOp>();
-		TSharedPtr<FFolderDragDropOp> FolderDragOp = DragDropEvent.GetOperationAs<FFolderDragDropOp>();
+		TSharedPtr<FActorDragDropOp> ActorDragOp = nullptr;
+		TSharedPtr<FFolderDragDropOp> FolderDragOp = nullptr;
+
+		if (const TSharedPtr<FCompositeDragDropOp> CompositeDragOp = DragDropEvent.GetOperationAs<FCompositeDragDropOp>())
+		{
+			ActorDragOp = CompositeDragOp->GetSubOp<FActorDragDropOp>();
+			FolderDragOp = CompositeDragOp->GetSubOp<FFolderDragDropOp>();
+		}
+		else
+		{
+			ActorDragOp = DragDropEvent.GetOperationAs<FActorDragDropOp>();
+			FolderDragOp = DragDropEvent.GetOperationAs<FFolderDragDropOp>();
+		}
+		
 		if (ActorDragOp.IsValid() || FolderDragOp.IsValid())
 		{
 			const FGeometry LayerContentsHeaderGeometry = SWidget::FindChildGeometry( MyGeometry, LayerContentsHeader.ToSharedRef() );
@@ -194,7 +207,21 @@ protected:
 
 		bool bHandled = false;
 		TArray<TWeakObjectPtr<AActor>> ActorsToDrop;
-		TSharedPtr<FActorDragDropOp> ActorDragOp = DragDropEvent.GetOperationAs<FActorDragDropOp>();
+
+		TSharedPtr<FActorDragDropOp> ActorDragOp = nullptr;
+		TSharedPtr<FFolderDragDropOp> FolderDragOp = nullptr;
+
+		if (const TSharedPtr<FCompositeDragDropOp> CompositeDragOp = DragDropEvent.GetOperationAs<FCompositeDragDropOp>())
+		{
+			ActorDragOp = CompositeDragOp->GetSubOp<FActorDragDropOp>();
+			FolderDragOp = CompositeDragOp->GetSubOp<FFolderDragDropOp>();
+		}
+		else
+		{
+			ActorDragOp = DragDropEvent.GetOperationAs<FActorDragDropOp>();
+			FolderDragOp = DragDropEvent.GetOperationAs<FFolderDragDropOp>();
+		}
+
 		if (ActorDragOp.IsValid())
 		{
 			const FGeometry LayerContentsHeaderGeometry = SWidget::FindChildGeometry(MyGeometry, LayerContentsHeader.ToSharedRef());
@@ -213,7 +240,6 @@ protected:
 			}
 		}
 
-		TSharedPtr<FFolderDragDropOp> FolderDragOp = DragDropEvent.GetOperationAs<FFolderDragDropOp>();
 		if (FolderDragOp.IsValid())
 		{
 			if (UWorld* World = FolderDragOp->World.Get())
diff --git a/Engine/Source/Editor/Layers/Private/SLayersView.h b/Engine/Source/Editor/Layers/Private/SLayersView.h
--- a/Engine/Source/Editor/Layers/Private/SLayersView.h
+++ b/Engine/Source/Editor/Layers/Private/SLayersView.h
@@ -210,14 +210,26 @@ protected:
 		bool bHandled = false;
 		TArray<TWeakObjectPtr<AActor>> ActorsToDrop;
 
-		TSharedPtr<FActorDragDropOp> ActorDragOp = DragDropEvent.GetOperationAs<FActorDragDropOp>();
+		TSharedPtr<FActorDragDropOp> ActorDragOp = nullptr;
+		TSharedPtr<FFolderDragDropOp> FolderDragOp = nullptr;
+
+		if (const TSharedPtr<FCompositeDragDropOp> CompositeDragOp = DragDropEvent.GetOperationAs<FCompositeDragDropOp>())
+		{
+			ActorDragOp = CompositeDragOp->GetSubOp<FActorDragDropOp>();
+			FolderDragOp = CompositeDragOp->GetSubOp<FFolderDragDropOp>();
+		}
+		else
+		{
+			ActorDragOp = DragDropEvent.GetOperationAs<FActorDragDropOp>();
+			FolderDragOp = DragDropEvent.GetOperationAs<FFolderDragDropOp>();
+		}
+
 		if (ActorDragOp.IsValid())
 		{
 			ActorsToDrop = ActorDragOp->Actors;
 			bHandled = true;
 		}
 
-		TSharedPtr<FFolderDragDropOp> FolderDragOp = DragDropEvent.GetOperationAs<FFolderDragDropOp>();
 		if (FolderDragOp.IsValid())
 		{
 			if (UWorld* World = FolderDragOp->World.Get())
diff --git a/Engine/Source/Editor/Layers/Private/SLayersViewRow.cpp b/Engine/Source/Editor/Layers/Private/SLayersViewRow.cpp
--- a/Engine/Source/Editor/Layers/Private/SLayersViewRow.cpp
+++ b/Engine/Source/Editor/Layers/Private/SLayersViewRow.cpp
@@ -8,6 +8,7 @@
 #include "EditorStyleSet.h"
 #include "DragAndDrop/ActorDragDropOp.h"
 #include "DragAndDrop/FolderDragDropOp.h"
+#include "DragAndDrop/CompositeDragDropOp.h"
 #include "Widgets/Text/SInlineEditableTextBlock.h"
 #include "EditorActorFolders.h"
 #include "Engine/World.h"
@@ -118,14 +119,26 @@ void SLayersViewRow::OnDragLeave(const FDragDropEvent& DragDropEvent)
 
 FReply SLayersViewRow::OnDragOver(const FGeometry& MyGeometry, const FDragDropEvent& DragDropEvent)
 {
-	TArray<TWeakObjectPtr<AActor>> Actors;
-	TSharedPtr<FActorDragDropOp> ActorDragOp = DragDropEvent.GetOperationAs<FActorDragDropOp>();
+	TArray<TWeakObjectPtr<AActor>> Actors; 
+	TSharedPtr<FActorDragDropOp> ActorDragOp = nullptr;
+	TSharedPtr<FFolderDragDropOp> FolderDragOp = nullptr;
+
+	if (const TSharedPtr<FCompositeDragDropOp> CompositeDragOp = DragDropEvent.GetOperationAs<FCompositeDragDropOp>())
+	{
+		ActorDragOp = CompositeDragOp->GetSubOp<FActorDragDropOp>();
+		FolderDragOp = CompositeDragOp->GetSubOp<FFolderDragDropOp>();
+	}
+	else
+	{
+		ActorDragOp = DragDropEvent.GetOperationAs<FActorDragDropOp>();
+		FolderDragOp = DragDropEvent.GetOperationAs<FFolderDragDropOp>();
+	}
+
 	if (ActorDragOp.IsValid() && ActorDragOp->Actors.Num() > 0)
 	{
 		Actors = ActorDragOp->Actors;
 	}
 
-	TSharedPtr<FFolderDragDropOp> FolderDragOp = DragDropEvent.GetOperationAs<FFolderDragDropOp>();
 	if (FolderDragOp.IsValid())
 	{
 		if (UWorld* World = FolderDragOp->World.Get())
@@ -180,15 +193,12 @@ FReply SLayersViewRow::OnDrop(const FGeometry& MyGeometry, const FDragDropEvent&
 	bool bHandled = false;
 	TArray<TWeakObjectPtr<AActor>> ActorsToDrop;
 
-	TSharedPtr< FActorDragDropOp > ActorDragOp = DragDropEvent.GetOperationAs<FActorDragDropOp>();
-	if (ActorDragOp.IsValid())
+	if (const TSharedPtr<FActorDragDropOp> ActorDragOp = DragDropEvent.GetOperationAs<FActorDragDropOp>())
 	{
 		ActorsToDrop = ActorDragOp->Actors;
 		bHandled = true;
 	}
-
-	TSharedPtr<FFolderDragDropOp> FolderDragOp = DragDropEvent.GetOperationAs<FFolderDragDropOp>();
-	if (FolderDragOp.IsValid())
+	else if (const TSharedPtr<FFolderDragDropOp> FolderDragOp = DragDropEvent.GetOperationAs<FFolderDragDropOp>())
 	{
 		if (UWorld* World = FolderDragOp->World.Get())
 		{
@@ -196,6 +206,22 @@ FReply SLayersViewRow::OnDrop(const FGeometry& MyGeometry, const FDragDropEvent&
 			bHandled = true;
 		}
 	}
+	else if (const TSharedPtr<FCompositeDragDropOp> CompositeDragOp = DragDropEvent.GetOperationAs<FCompositeDragDropOp>())
+	{
+		if (const TSharedPtr<FActorDragDropOp> ActorSubOp = CompositeDragOp->GetSubOp<FActorDragDropOp>())
+		{
+			ActorsToDrop = ActorSubOp->Actors;
+			bHandled = true;
+		}
+		if (const TSharedPtr<FFolderDragDropOp> FolderSubOp = CompositeDragOp->GetSubOp<FFolderDragDropOp>())
+		{
+			if (UWorld* World = FolderDragOp->World.Get())
+			{
+				FActorFolders::GetWeakActorsFromFolders(*World, FolderSubOp->Folders, ActorsToDrop);
+				bHandled = true;
+			}
+		}
+	}
 
 	if (ActorsToDrop.Num() > 0)
 	{
@@ -215,12 +241,29 @@ FSlateColor SLayersViewRow::GetColorAndOpacity() const
 	bool bCanAcceptDrop = false;
 	TSharedPtr<FDragDropOperation> DragDropOp = FSlateApplication::Get().GetDragDroppingContent();
 
-	if (DragDropOp.IsValid() && (DragDropOp->IsOfType<FActorDragDropOp>() || DragDropOp->IsOfType<FFolderDragDropOp>()))
+	if (DragDropOp.IsValid() && (DragDropOp->IsOfType<FActorDragDropOp>() || DragDropOp->IsOfType<FFolderDragDropOp>() || DragDropOp->IsOfType<FCompositeDragDropOp>()))
 	{
 		TArray<TWeakObjectPtr<AActor>> DraggedActors;
 
-		TSharedPtr<FActorDragDropOp> ActorDragOp = DragDropOp->CastTo<FActorDragDropOp>();
-		TSharedPtr<FFolderDragDropOp> FolderDragOp = DragDropOp->CastTo<FFolderDragDropOp>();
+		TSharedPtr<FActorDragDropOp> ActorDragOp;
+		TSharedPtr<FFolderDragDropOp> FolderDragOp;
+
+		if (DragDropOp->IsOfType<FActorDragDropOp>())
+		{
+			ActorDragOp = StaticCastSharedPtr<FActorDragDropOp>(DragDropOp);
+		}
+		else if (DragDropOp->IsOfType<FFolderDragDropOp>())
+		{
+			FolderDragOp = StaticCastSharedPtr<FFolderDragDropOp>(DragDropOp);
+		}
+		else if (DragDropOp->IsOfType<FCompositeDragDropOp>())
+		{
+			if (const TSharedPtr<FCompositeDragDropOp> CompositeDrop = StaticCastSharedPtr<FCompositeDragDropOp>(DragDropOp))
+			{
+				ActorDragOp = CompositeDrop->GetSubOp<FActorDragDropOp>();
+				FolderDragOp = CompositeDrop->GetSubOp<FFolderDragDropOp>();
+			}
+		}
 
 		if (ActorDragOp.IsValid())
 		{
diff --git a/Engine/Source/Editor/SceneOutliner/Private/ActorBrowsingMode.cpp b/Engine/Source/Editor/SceneOutliner/Private/ActorBrowsingMode.cpp
--- a/Engine/Source/Editor/SceneOutliner/Private/ActorBrowsingMode.cpp
+++ b/Engine/Source/Editor/SceneOutliner/Private/ActorBrowsingMode.cpp
@@ -18,6 +18,7 @@
 #include "EditorFolderUtils.h"
 #include "ActorEditorUtils.h"
 #include "DragAndDrop/ActorDragDropOp.h"
+#include "DragAndDrop/ActorDragDropGraphEdOp.h"
 #include "DragAndDrop/FolderDragDropOp.h"
 #include "Logging/MessageLog.h"
 #include "SSocketChooser.h"
@@ -739,11 +740,17 @@ TSharedPtr<FDragDropOperation> FActorBrowsingMode::CreateDragDropOperation(const
 {
 	FSceneOutlinerDragDropPayload DraggedObjects(InTreeItems);
 
+	// If the drag contains only actors, we shortcut and create a simple FActorDragDropGraphEdOp rather than an FSceneOutlinerDragDrop composite op.
+	if (DraggedObjects.Has<FActorTreeItem>() && !DraggedObjects.Has<FFolderTreeItem>())
+	{
+		return FActorDragDropGraphEdOp::New(DraggedObjects.GetData<TWeakObjectPtr<AActor>>(SceneOutliner::FWeakActorSelector()));
+	}
+
 	TSharedPtr<FSceneOutlinerDragDropOp> OutlinerOp = MakeShareable(new FSceneOutlinerDragDropOp());
 
 	if (DraggedObjects.Has<FActorTreeItem>())
 	{
-		TSharedPtr<FActorDragDropOp> ActorOperation = MakeShareable(new FActorDragDropOp);
+		TSharedPtr<FActorDragDropOp> ActorOperation = MakeShareable(new FActorDragDropGraphEdOp);
 		ActorOperation->Init(DraggedObjects.GetData<TWeakObjectPtr<AActor>>(SceneOutliner::FWeakActorSelector()));
 		OutlinerOp->AddSubOp(ActorOperation);
 	}
diff --git a/Engine/Source/Editor/Sequencer/Private/SSequencer.cpp b/Engine/Source/Editor/Sequencer/Private/SSequencer.cpp
--- a/Engine/Source/Editor/Sequencer/Private/SSequencer.cpp
+++ b/Engine/Source/Editor/Sequencer/Private/SSequencer.cpp
@@ -53,6 +53,7 @@
 #include "DragAndDrop/ActorDragDropGraphEdOp.h"
 #include "DragAndDrop/ClassDragDropOp.h"
 #include "DragAndDrop/FolderDragDropOp.h"
+#include "DragAndDrop/CompositeDragDropOp.h"
 #include "Widgets/Input/SSearchBox.h"
 #include "SSequencerTreeView.h"
 #include "MovieSceneTrackEditor.h"
@@ -3010,32 +3011,47 @@ FReply SSequencer::OnDrop( const FGeometry& MyGeometry, const FDragDropEvent& Dr
 	{
 		if ( Operation->IsOfType<FAssetDragDropOp>() )
 		{
-			const auto& DragDropOp = Operation->CastTo<FAssetDragDropOp>();
+			const auto& DragDropOp = StaticCastSharedPtr<FAssetDragDropOp>(Operation);
 
 			OnAssetsDropped( *DragDropOp );
 			bWasDropHandled = true;
 		}
-		if( Operation->IsOfType<FClassDragDropOp>() )
+		else if( Operation->IsOfType<FClassDragDropOp>() )
 		{
-			const auto& DragDropOp = Operation->CastTo<FClassDragDropOp>();
+			const auto& DragDropOp = StaticCastSharedPtr<FClassDragDropOp>(Operation);
 
 			OnClassesDropped( *DragDropOp );
 			bWasDropHandled = true;
 		}
-		if( Operation->IsOfType<FActorDragDropOp>() )
+		else if( Operation->IsOfType<FActorDragDropOp>() )
 		{
-			const auto& DragDropOp = Operation->CastTo<FActorDragDropOp>();
+			const auto& DragDropOp = StaticCastSharedPtr<FActorDragDropOp>(Operation);
 
 			OnActorsDropped( *DragDropOp );
 			bWasDropHandled = true;
 		}
-		if (Operation->IsOfType<FFolderDragDropOp>())
+		else if (Operation->IsOfType<FFolderDragDropOp>())
 		{
-			const auto& DragDropOp = Operation->CastTo<FFolderDragDropOp>();
+			const auto& DragDropOp = StaticCastSharedPtr<FFolderDragDropOp>(Operation);
 
 			OnFolderDropped(*DragDropOp);
 			bWasDropHandled = true;
 		}
+		else if (Operation->IsOfType<FCompositeDragDropOp>())
+		{
+			const TSharedPtr<FCompositeDragDropOp> CompositeOp = StaticCastSharedPtr<FCompositeDragDropOp>(Operation);
+			if (const TSharedPtr<FActorDragDropOp> ActorDragDropOp = CompositeOp->GetSubOp<FActorDragDropOp>())
+			{
+				OnActorsDropped(*ActorDragDropOp);
+				bWasDropHandled = true;
+			}
+			if (const TSharedPtr<FFolderDragDropOp> FolderDragDropOp = CompositeOp->GetSubOp<FFolderDragDropOp>())
+			{
+				OnFolderDropped(*FolderDragDropOp);
+				bWasDropHandled = true;
+			}
+			
+		}
 	}
 
 	return bWasDropHandled ? FReply::Handled() : FReply::Unhandled();
diff --git a/Engine/Source/Editor/UnrealEd/Public/DragAndDrop/CompositeDragDropOp.h b/Engine/Source/Editor/UnrealEd/Public/DragAndDrop/CompositeDragDropOp.h
--- a/Engine/Source/Editor/UnrealEd/Public/DragAndDrop/CompositeDragDropOp.h
+++ b/Engine/Source/Editor/UnrealEd/Public/DragAndDrop/CompositeDragDropOp.h
@@ -9,11 +9,7 @@
 class FCompositeDragDropOp : public FDecoratedDragDropOp
 {
 public:
-	static const FString& GetTypeId() { static FString Type = TEXT("FCompositeDragDropOp"); return Type; }
-	virtual bool IsOfTypeImpl(const FString& Type) const override 
-	{
-		return GetTypeId() == Type || FDecoratedDragDropOp::IsOfTypeImpl(Type) || GetSubOpPtr(Type) != nullptr; 
-	}
+	DRAG_DROP_OPERATOR_TYPE(FCompositeDragDropOp, FDecoratedDragDropOp)
 
 	FCompositeDragDropOp() : FDecoratedDragDropOp() {}
 
@@ -23,6 +19,19 @@ public:
 		SubOps.Add(SubOp);
 	}
 
+	template <typename T>
+	TSharedPtr<T> HasSubOp()
+	{
+		for (auto& SubOp : SubOps)
+		{
+			if (SubOp->IsOfType<T>())
+			{
+				return true;
+			}
+		}
+		return false;
+	}
+
 	template <typename T>
 	TSharedPtr<T> GetSubOp()
 	{
@@ -49,21 +58,6 @@ public:
 		return nullptr;
 	}
 
-	virtual TSharedPtr<FDragDropOperation> ConvertTo(const FString& TypeId) override
-	{
-		// Attempt to convert this before getting subops
-		if (FDecoratedDragDropOp::IsOfTypeImpl(TypeId))
-		{
-			return AsShared();
-		}
-		else
-		{
-			// Will be nullptr if failed to convert
-			TSharedPtr<FDragDropOperation> SubOpAs = GetSubOpPtr(TypeId);
-			return SubOpAs;
-		}
-	}
-
 	virtual void ResetToDefaultToolTip() override
 	{
 		FDecoratedDragDropOp::ResetToDefaultToolTip();
@@ -72,7 +66,7 @@ public:
 		{
 			if (SubOp->IsOfType<FDecoratedDragDropOp>())
 			{
-				auto DecoratedSubOp = SubOp->CastTo<FDecoratedDragDropOp>();
+				auto DecoratedSubOp = StaticCastSharedPtr<FDecoratedDragDropOp>(SubOp);
 				if (DecoratedSubOp.IsValid())
 				{
 					DecoratedSubOp->ResetToDefaultToolTip();
diff --git a/Engine/Source/Runtime/SlateCore/Public/Input/DragAndDrop.h b/Engine/Source/Runtime/SlateCore/Public/Input/DragAndDrop.h
--- a/Engine/Source/Runtime/SlateCore/Public/Input/DragAndDrop.h
+++ b/Engine/Source/Runtime/SlateCore/Public/Input/DragAndDrop.h
@@ -37,17 +37,6 @@ public:
 		return IsOfTypeImpl(TType::GetTypeId());
 	}
 
-	/** Cast this drag and drop operation to the specified template type */
-	template <class OperationType> TSharedPtr<OperationType> CastTo()
-	{
-		TSharedPtr<FDragDropOperation> Converted(this->ConvertTo(OperationType::GetTypeId()));
-		if (Converted != nullptr)
-		{
-			return StaticCastSharedPtr<OperationType>(Converted);
-		}
-		return TSharedPtr<OperationType>();
-	}
-
 	/**
 	* Checks if this drag and drop operation is affected by a given PointerEvent.
 	*
diff --git a/Engine/Source/Runtime/SlateCore/Public/Input/DragAndDrop.inl b/Engine/Source/Runtime/SlateCore/Public/Input/DragAndDrop.inl
--- a/Engine/Source/Runtime/SlateCore/Public/Input/DragAndDrop.inl
+++ b/Engine/Source/Runtime/SlateCore/Public/Input/DragAndDrop.inl
@@ -10,7 +10,7 @@ TSharedPtr<OperationType> FDragDropEvent::GetOperationAs() const
 {
 	if (Content.IsValid() && Content->IsOfType<OperationType>())
 	{
-		return Content->CastTo<OperationType>();
+		return StaticCastSharedPtr<OperationType>(Content);
 	}
 	else
 	{
