Tweaked comment in LayerCollectionViewModel.cpp to remove unacceptable word.

diff --git a/Engine/Source/Editor/Layers/Private/LayerCollectionViewModel.cpp b/Engine/Source/Editor/Layers/Private/LayerCollectionViewModel.cpp
--- a/Engine/Source/Editor/Layers/Private/LayerCollectionViewModel.cpp
+++ b/Engine/Source/Editor/Layers/Private/LayerCollectionViewModel.cpp
@@ -296,7 +296,7 @@ void FLayerCollectionViewModel::DestructivelyPurgeInvalidViewModels( TArray< TWe
 		const auto Layer = LayerViewModel->GetDataSource(); 
 
 		//Remove any viewmodels with invalid datasources or whose datasources 
-		//are no longer in the master list of layers
+		//are no longer in the list of layers
 		if( !Layer.IsValid() || InLayers.Remove( Layer ) == 0 )
 		{
 			AllLayerViewModels.RemoveAt( LayerIndex );
